Port MessageBus tester plugin into the UE experimental folder. This plugin is designed to provide test cases that can be used to exercise UDP message transport. The test cases help analyze use cases for messaging and make sure the low level code is functioning optimally.  This is only designed for testing and debugging at the moment and may be changed or removed in the future.

diff --git a/Engine/Plugins/Animation/LiveLinkHub/LiveLinkHub.uplugin b/Engine/Plugins/Animation/LiveLinkHub/LiveLinkHub.uplugin
--- a/Engine/Plugins/Animation/LiveLinkHub/LiveLinkHub.uplugin
+++ b/Engine/Plugins/Animation/LiveLinkHub/LiveLinkHub.uplugin
@@ -28,6 +28,10 @@
 		{
 			"Name": "UdpMessaging",
 			"Enabled": true
+		},
+        {
+			"Name": "MessageBusTester",
+			"Enabled": true
 		}
 	],
 	"Modules": [
diff --git a/Engine/Plugins/Animation/LiveLinkHub/Source/LiveLinkHub/LiveLinkHub.Build.cs b/Engine/Plugins/Animation/LiveLinkHub/Source/LiveLinkHub/LiveLinkHub.Build.cs
--- a/Engine/Plugins/Animation/LiveLinkHub/Source/LiveLinkHub/LiveLinkHub.Build.cs
+++ b/Engine/Plugins/Animation/LiveLinkHub/Source/LiveLinkHub/LiveLinkHub.Build.cs
@@ -48,6 +48,7 @@ public class LiveLinkHub : ModuleRules
 				"UnrealEd",
 				"OutputLog",
 				"NamingTokens",
+				"MessageBusTesterEditor"
 			});
 
 		if (Target.Configuration != UnrealTargetConfiguration.Shipping)
diff --git a/Engine/Plugins/Animation/LiveLinkHub/Source/LiveLinkHub/Private/LiveLinkHubApplication.cpp b/Engine/Plugins/Animation/LiveLinkHub/Source/LiveLinkHub/Private/LiveLinkHubApplication.cpp
--- a/Engine/Plugins/Animation/LiveLinkHub/Source/LiveLinkHub/Private/LiveLinkHubApplication.cpp
+++ b/Engine/Plugins/Animation/LiveLinkHub/Source/LiveLinkHub/Private/LiveLinkHubApplication.cpp
@@ -24,6 +24,7 @@
 #		include "ISessionFrontendModule.h"
 #	endif
 #	include "ISlateReflectorModule.h"
+#	include "IMessageBusTesterEditorModule.h"
 #endif
 
 #define LOCTEXT_NAMESPACE "LiveLinkHubApplication"
@@ -581,8 +582,21 @@ void FLiveLinkHubApplication::CreateToolsMenu(FMenuBuilder& MenuBuilder, UToolMe
 				SlateReflector.DisplayWidgetReflector();
 			})
 	);
+	MenuBuilder.AddMenuEntry(
+		LOCTEXT("FileMenu_Development_MesageBusTester", "Message Bus Tester"),
+		FText(),
+		FSlateIcon(),
+		FExecuteAction::CreateStatic([]()
+			{
+				if (IMessageBusTesterEditorModule::IsAvailable())
+				{
+					IMessageBusTesterEditorModule::Get().DisplayMessageBusTester();
+				}
+			})
+	);
 	MenuBuilder.EndSection();
 
+
 #endif /*UE_BUILD_SHIPPING*/
 }
 
diff --git a/Engine/Plugins/Experimental/MessageBusTester/MessageBusTester.uplugin b/Engine/Plugins/Experimental/MessageBusTester/MessageBusTester.uplugin
new file mode 100644
--- /dev/null
+++ b/Engine/Plugins/Experimental/MessageBusTester/MessageBusTester.uplugin
@@ -0,0 +1,37 @@
+{
+	"FileVersion": 3,
+	"Version": 1,
+	"VersionName": "1.0",
+	"FriendlyName": "MessageBus Tester",
+	"Description": "Plugin to test and monitor message bus reliability",
+	"Category": "Messaging",
+	"CreatedBy": "",
+	"CreatedByURL": "",
+	"DocsURL": "",
+	"MarketplaceURL": "",
+	"SupportURL": "",
+	"CanContainContent": false,
+	"IsBetaVersion": true,
+	"Installed": false,
+	"SupportedPrograms": [ "MessageBusTesterApp" ],
+	"Plugins": [
+		{
+			"Name": "UdpMessaging",
+			"Enabled": true
+		}
+	],
+	"Modules": [
+		{
+			"Name": "MessageBusTester",
+			"Type": "Runtime",
+			"LoadingPhase": "Default",
+			"AllowlistPrograms": [ "MessageBusTesterApp" ]
+		},
+		{
+			"Name": "MessageBusTesterEditor",
+          	"Type": "UncookedOnly",
+			"LoadingPhase": "Default",
+			"AllowlistPrograms": [ "MessageBusTesterApp" ]
+		}
+	]
+}
diff --git a/Engine/Plugins/Experimental/MessageBusTester/Source/MessageBusTester/MessageBusTester.Build.cs b/Engine/Plugins/Experimental/MessageBusTester/Source/MessageBusTester/MessageBusTester.Build.cs
new file mode 100644
--- /dev/null
+++ b/Engine/Plugins/Experimental/MessageBusTester/Source/MessageBusTester/MessageBusTester.Build.cs
@@ -0,0 +1,41 @@
+// Copyright Epic Games, Inc. All Rights Reserved.
+
+using UnrealBuildTool;
+
+public class MessageBusTester : ModuleRules
+{
+	public MessageBusTester(ReadOnlyTargetRules Target) : base(Target)
+	{
+		PCHUsage = PCHUsageMode.UseExplicitOrSharedPCHs;
+
+		bUseUnity = false;
+		PCHUsage = PCHUsageMode.NoPCHs;
+
+		PublicDependencyModuleNames.AddRange(
+			new string[]
+			{
+                "DeveloperSettings",
+				"MessagingCommon",
+				"UdpMessaging"
+			}
+		);
+
+		PrivateDependencyModuleNames.AddRange(
+			new string[]
+			{
+				"Core",
+				"CoreUObject",
+				"Messaging",
+			}
+		);
+
+		if (Target.bBuildEditor)
+		{
+			PrivateDependencyModuleNames.AddRange(
+				new string[] {
+					"MessageLog",
+				}
+			);
+		}
+	}
+}
diff --git a/Engine/Plugins/Experimental/MessageBusTester/Source/MessageBusTester/Private/DiscoveredTester.cpp b/Engine/Plugins/Experimental/MessageBusTester/Source/MessageBusTester/Private/DiscoveredTester.cpp
new file mode 100644
--- /dev/null
+++ b/Engine/Plugins/Experimental/MessageBusTester/Source/MessageBusTester/Private/DiscoveredTester.cpp
@@ -0,0 +1,98 @@
+// Copyright Epic Games, Inc. All Rights Reserved.
+
+#include "DiscoveredTester.h"
+
+#include "MessageBusTesterModule.h"
+#include "IMessageBusTesterLogger.h"
+#include "MessageBusTesterSettings.h"
+
+
+#define LOCTEXT_NAMESPACE "DiscoveredTester"
+
+void FDiscoveredTester::UpdateConnectionState()
+{
+	if (ConnectionState != EDiscoveredTesterConnectionState::Lost)
+	{
+		if ((FPlatformTime::Seconds() - LastUnreliableKeepAliveMessageTime > GetDefault<UMessageBusTesterSettings>()->TimeoutInterval)
+			&& (FPlatformTime::Seconds() - LastReliableKeepAliveMessageTime > GetDefault<UMessageBusTesterSettings>()->TimeoutInterval))
+		{
+			UE_LOG(LogMessageBusTester, Log, TEXT("Tester lost, %s"), *Descriptor.FriendlyName.ToString());
+			IMessageBusTesterModule::Get().GetLogger().Log(Descriptor.FriendlyName, FString::Printf(TEXT("Tester lost connection for missed keepalive")), EMessageSeverity::Warning);
+			ConnectionState = EDiscoveredTesterConnectionState::Lost;
+		}
+	}
+	else if (ConnectionState == EDiscoveredTesterConnectionState::Lost)
+	{
+		if ((FPlatformTime::Seconds() - LastUnreliableKeepAliveMessageTime <= GetDefault<UMessageBusTesterSettings>()->TimeoutInterval)
+			|| (FPlatformTime::Seconds() - LastReliableKeepAliveMessageTime <= GetDefault<UMessageBusTesterSettings>()->TimeoutInterval))
+		{
+			UE_LOG(LogMessageBusTester, Log, TEXT("Tester back alive, %s"), *Descriptor.FriendlyName.ToString());
+			IMessageBusTesterModule::Get().GetLogger().Log(Descriptor.FriendlyName, FString::Printf(TEXT("Tester recovered")), EMessageSeverity::Info);
+			ConnectionState = EDiscoveredTesterConnectionState::Connected;
+		}
+	}
+}
+
+void FDiscoveredTester::NotifyPayloadSent(const FTestPayloadMessage& PayloadMessage)
+{
+	SentPayloads.AddUnique({PayloadMessage});
+}
+
+void FDiscoveredTester::NotifyAckReceived(const FTestPayloadReceptionMessage& AckMessage)
+{
+	const double ReceivedTime = FPlatformTime::Seconds();
+	const int32 PayloadId = AckMessage.PayloadId;
+	const int32 SentPayloadIndex = SentPayloads.IndexOfByPredicate([PayloadId](const FPayloadInfo& Other) { return Other.PayloadId == PayloadId; });
+	if (SentPayloads.IsValidIndex(SentPayloadIndex))
+	{
+		FPayloadInfo ReceivedPayload = SentPayloads[SentPayloadIndex];
+		const double Elapsed = FMath::Max(ReceivedTime - ReceivedPayload.SentTime, UE_KINDA_SMALL_NUMBER);
+		const double MegaByte = 1024.0f * 1024.0f;
+		const double BandwithMegaBytes = static_cast<double>(ReceivedPayload.PayloadSize) / MegaByte / Elapsed;
+		IMessageBusTesterModule::Get().GetLogger().Log(Descriptor.FriendlyName, FString::Printf(TEXT("Packet %d acknowledged after %0.5f (%0.1f MB/s)"), ReceivedPayload.PayloadId, Elapsed, BandwithMegaBytes), EMessageSeverity::Info);
+		
+		ReceivedPayload.AcknowledgedTime = ReceivedTime;
+		SentPayloads.RemoveAtSwap(SentPayloadIndex);
+		AcknowledgedPayloads.AddUnique(MoveTemp(ReceivedPayload));
+
+	}
+	else
+	{
+		ensureAlways(false);
+	}
+}
+
+void FDiscoveredTester::NotifyKeepAliveReceived(bool bIsReliable)
+{
+	const double ReceptionTime = FPlatformTime::Seconds();
+	if (bIsReliable)
+	{
+		KeepAliveStatistics.LastReliableKeepAliveReceptionInterval = ReceptionTime - LastReliableKeepAliveMessageTime;
+		KeepAliveStatistics.MinReliableKeepAliveInterval = FMath::Min(KeepAliveStatistics.MinReliableKeepAliveInterval, KeepAliveStatistics.LastReliableKeepAliveReceptionInterval);
+		KeepAliveStatistics.MaxReliableKeepAliveInterval = FMath::Max(KeepAliveStatistics.MaxReliableKeepAliveInterval, KeepAliveStatistics.LastReliableKeepAliveReceptionInterval);
+		AverageReliableKeepAliveIntervalWindow.PushValue(KeepAliveStatistics.LastReliableKeepAliveReceptionInterval);
+		KeepAliveStatistics.AverageReliableKeepAliveInterval = AverageReliableKeepAliveIntervalWindow.GetAverage();
+		LastReliableKeepAliveMessageTime = ReceptionTime;
+		
+		UE_LOG(LogMessageBusTester, Verbose, TEXT("Received Reliable KeepAlive from '%s' after %0.2f"), *Descriptor.FriendlyName.ToString(), KeepAliveStatistics.LastReliableKeepAliveReceptionInterval);
+	}
+	else
+	{
+		KeepAliveStatistics.LastUnreliableKeepAliveReceptionInterval = ReceptionTime - LastUnreliableKeepAliveMessageTime;
+		KeepAliveStatistics.MinUnreliableKeepAliveInterval = FMath::Min(KeepAliveStatistics.MinUnreliableKeepAliveInterval, KeepAliveStatistics.LastUnreliableKeepAliveReceptionInterval);
+		KeepAliveStatistics.MaxUnreliableKeepAliveInterval = FMath::Max(KeepAliveStatistics.MaxUnreliableKeepAliveInterval, KeepAliveStatistics.LastUnreliableKeepAliveReceptionInterval);
+		AverageUnreliableKeepAliveIntervalWindow.PushValue(KeepAliveStatistics.LastUnreliableKeepAliveReceptionInterval);
+		KeepAliveStatistics.AverageUnreliableKeepAliveInterval = AverageUnreliableKeepAliveIntervalWindow.GetAverage();
+		LastUnreliableKeepAliveMessageTime = ReceptionTime;
+
+		UE_LOG(LogMessageBusTester, Verbose, TEXT("Received Unreliable KeepAlive from '%s' after %0.2f"), *Descriptor.FriendlyName.ToString(), KeepAliveStatistics.LastUnreliableKeepAliveReceptionInterval);
+	}
+}
+
+bool FDiscoveredTester::IsPendingPayloadAck(int32 PayloadId) const
+{
+	return  INDEX_NONE != SentPayloads.IndexOfByPredicate([PayloadId](const FPayloadInfo& Other) { return Other.PayloadId == PayloadId; });
+
+}
+
+#undef LOCTEXT_NAMESPACE
diff --git a/Engine/Plugins/Experimental/MessageBusTester/Source/MessageBusTester/Private/MessageBusTester.cpp b/Engine/Plugins/Experimental/MessageBusTester/Source/MessageBusTester/Private/MessageBusTester.cpp
new file mode 100644
--- /dev/null
+++ b/Engine/Plugins/Experimental/MessageBusTester/Source/MessageBusTester/Private/MessageBusTester.cpp
@@ -0,0 +1,496 @@
+// Copyright Epic Games, Inc. All Rights Reserved.
+
+#include "MessageBusTester.h"
+
+#include "Containers/Ticker.h"
+#include "MessageEndpoint.h"
+#include "MessageEndpointBuilder.h"
+#include "Misc/App.h"
+#include "Misc/CoreDelegates.h"
+#include "MessageBusTesterModule.h"
+#include "MessageBusTesterSettings.h"
+#include "Async/Async.h"
+
+
+namespace MessageBusTesterUtils
+{
+	static constexpr int32 SupportedMessageBusTesterVersion = 1;
+	static constexpr int32 KeepAliveIntervalMilliSeconds = 1000;
+}
+
+FMessageBusTester::~FMessageBusTester()
+{
+	FCoreDelegates::OnPreExit.RemoveAll(this);
+	
+	Stop();
+}
+
+void FMessageBusTester::AddTestPlanItem(FTestPlanItem NewItem)
+{
+	TestPlan.TestPlanItems.Emplace(MoveTemp(NewItem));
+	OnTestPlanChanged().Broadcast();
+}
+
+void FMessageBusTester::RemoveTestPlanItem(int32 Index)
+{
+	if (TestPlan.TestPlanItems.IsValidIndex(Index))
+	{
+		TestPlan.TestPlanItems.RemoveAt(Index);
+		OnTestPlanChanged().Broadcast();
+	}
+}
+
+bool FMessageBusTester::StartTest()
+{
+	if (State == EMessageBusTesterState::Idle)
+	{
+		SetupTestPlan();
+
+		//Notify other testers we are starting
+		FTesterStartTestMessage* NewMsg = FMessageEndpoint::MakeMessage<FTesterStartTestMessage>();
+		NewMsg->TestPlan = TestPlan;
+		SendMessage(NewMsg, EMessageBusTesterMessageFlags::Reliable);
+
+		return true;
+	}
+
+	return false;
+}
+
+bool FMessageBusTester::StopTest()
+{
+	if (State == EMessageBusTesterState::Active)
+	{
+		//Notify other testers we are stopping the test
+		FTesterStopTestMessage* NewMsg = FMessageEndpoint::MakeMessage<FTesterStopTestMessage>();
+		SendMessage(NewMsg, EMessageBusTesterMessageFlags::Reliable);
+
+		State = EMessageBusTesterState::Idle;
+
+		return true;
+	}
+
+	return false;
+}
+
+void FMessageBusTester::Initialize()
+{
+	FCoreDelegates::OnPreExit.AddRaw(this, &FMessageBusTester::OnPreExit);
+
+	Identifier = FGuid::NewGuid();
+
+	KeepAliveEvent = FGenericPlatformProcess::GetSynchEventFromPool();
+
+	CachedDiscoveryMessage.Descriptor = MessageBusTesterUtils::GetInstanceDescriptor();
+
+	Start();
+}
+
+void FMessageBusTester::Start()
+{
+	UE_LOG(LogMessageBusTester, Log, TEXT("Starting MessageBusTester '%s'"), *CachedDiscoveryMessage.Descriptor.FriendlyName.ToString());
+
+	const FString MessageEndpointName = TEXT("MessageBusTester");
+	TesterEndpoint = FMessageEndpoint::Builder(*MessageEndpointName)
+		.ReceivingOnThread(ENamedThreads::Type::GameThread)
+		.Handling<FTesterDiscoveryMessage>(this, &FMessageBusTester::HandleDiscoveryMessage)
+		.Handling<FTesterKeepAliveMessage>(this, &FMessageBusTester::HandleKeepAlive)
+		.Handling<FTesterStartTestMessage>(this, &FMessageBusTester::HandleStartTest)
+		.Handling<FTesterStopTestMessage>(this, &FMessageBusTester::HandleStopTest)
+		.Handling<FTestPayloadMessage>(this, &FMessageBusTester::HandlePayloadMessage)
+		.Handling<FTestPayloadReceptionMessage>(this, &FMessageBusTester::HandlePayloadReceptionMessage);
+
+	if (TesterEndpoint.IsValid())
+	{
+		TesterEndpoint->Subscribe<FTesterDiscoveryMessage>();
+
+		bShouldSendKeepAlive = true;
+		KeepAliveTaskFuture = Async(EAsyncExecution::TaskGraph, [this]()
+		{ 
+			KeepAliveTask();
+		});
+
+		AddTestPlanItem({ 10 * 1024 * 1024, 0.0f});
+	}
+
+	TickerHandle = FTSTicker::GetCoreTicker().AddTicker(FTickerDelegate::CreateRaw(this, &FMessageBusTester::Tick), 0.0f);
+}
+
+void FMessageBusTester::Stop()
+{
+	if (KeepAliveEvent)
+	{
+		KeepAliveEvent->Trigger();
+	}
+
+	bShouldSendKeepAlive = false;
+	if (KeepAliveTaskFuture.IsValid())
+	{
+		KeepAliveTaskFuture.WaitFor(FTimespan::FromMilliseconds(100.0));
+	}
+
+	if (TickerHandle.IsValid())
+	{
+		FTSTicker::GetCoreTicker().RemoveTicker(TickerHandle);
+		TickerHandle.Reset();
+	}
+
+	//Let providers know we're out
+	if (TesterEndpoint.IsValid())
+	{
+		// Disable the Endpoint message handling since the message could keep it alive a bit.
+		TesterEndpoint->Disable();
+		TesterEndpoint.Reset();
+	}
+
+	InvalidTesters.Empty();
+
+	if (KeepAliveEvent)
+	{
+		FGenericPlatformProcess::ReturnSynchEventToPool(KeepAliveEvent);
+		KeepAliveEvent = nullptr;
+	}
+}
+
+bool FMessageBusTester::Tick(float DeltaTime)
+{
+	UpdateDiscoveredTester();
+	SendDiscoveryMessage();
+
+	if (State == EMessageBusTesterState::Active)
+	{
+		UpdateTestSequence();
+	}
+
+	return true;
+}
+
+void FMessageBusTester::UpdateDiscoveredTester()
+{
+	bool bTesterStateChanged = false;
+
+	FScopeLock Lock(&DiscoveredTesterCriticalSection);
+
+	for (TSharedPtr<FDiscoveredTester, ESPMode::ThreadSafe>& Tester : DiscoveredTesters)
+	{
+		Tester->UpdateConnectionState();
+	}
+}
+
+void FMessageBusTester::SendDiscoveryMessage()
+{
+	if (TesterEndpoint && (FPlatformTime::Seconds() - LastSentDiscoveryMessage) > GetDefault<UMessageBusTesterSettings>()->DiscoveryMessageInterval)
+	{
+		//Broadcast discovery signal
+		FTesterDiscoveryMessage* DiscoveryMessage = FMessageEndpoint::MakeMessage<FTesterDiscoveryMessage>();
+		DiscoveryMessage->Identifier = Identifier;
+		DiscoveryMessage->Descriptor = CachedDiscoveryMessage.Descriptor;
+
+		TesterEndpoint->Publish(DiscoveryMessage, EMessageScope::Network, FTimespan::Zero(), FDateTime::MaxValue());
+
+		LastSentDiscoveryMessage = FPlatformTime::Seconds();
+	}
+}
+
+void FMessageBusTester::OnPreExit()
+{
+	Stop();
+}
+
+void FMessageBusTester::UpdateTestSequence()
+{
+	FScopeLock Lock(&DiscoveredTesterCriticalSection);
+
+	if (DiscoveredTesters.Num() <= 0)
+	{
+		return;
+	}
+
+	//Go through each discovered tester and send the next test plan payload if last was was received
+	for (const TSharedPtr<FDiscoveredTester, ESPMode::ThreadSafe>& Tester : DiscoveredTesters)
+	{
+		if (Tester->ConnectionState != EDiscoveredTesterConnectionState::Lost
+		&& Tester->State == EMessageBusTesterState::Active)
+		{
+			for (int32 Index = 0; Index < TestPlanItemTrackers.Num(); ++Index)
+			{
+				FTestPlanItemTracker& Test = TestPlanItemTrackers[Index];
+				if ((FPlatformTime::Seconds() - Test.LastPayloadSentTimestamp) > Test.ItemDescription.IntervalSeconds)
+				{
+					bool bCanSendPayload = true;
+					if (int32* SentPayloadId = Test.SentPayloadIds.Find(Tester->Identifier))
+					{
+						if (Tester->IsPendingPayloadAck(*SentPayloadId))
+						{
+							bCanSendPayload = false;
+						}
+					}
+
+					if (bCanSendPayload)
+					{
+						static int32 NextPayloadId = 0;
+
+						//Interval elapsed, resend test
+						FTestPayloadMessage* NewPayload = FMessageEndpoint::MakeMessage<FTestPayloadMessage>();
+						NewPayload->Payload.SetNum(Test.ItemDescription.NumBytes);
+						NewPayload->TestPlanIndex = Index;
+
+						const int32 PayloadId = ++NextPayloadId;
+						NewPayload->PayloadId = PayloadId;
+						Tester->NotifyPayloadSent(*NewPayload);
+
+						SendMessage(NewPayload, {Tester->Address}, EMessageBusTesterMessageFlags::Reliable);
+
+						//Timestamp starts after pushing the packet in the stack
+						Test.LastPayloadSentTimestamp = FPlatformTime::Seconds();
+
+						//Mark all pending reception to false
+						Test.SentPayloadIds.FindOrAdd(Tester->Identifier) = PayloadId;
+
+						UE_LOG(LogMessageBusTester, Log, TEXT("Sending testplan item %d (payload %d of size %d)"), Index, PayloadId, Test.ItemDescription.NumBytes);
+					}
+				}
+				
+			}
+		}
+	}
+}
+
+void FMessageBusTester::KeepAliveTask()
+{
+	{
+		FScopeLock Lock(&DiscoveredTesterCriticalSection);
+
+		static uint64 KeepAliveNumber = 0;
+		FTesterKeepAliveMessage* NewKeepAlive = FMessageEndpoint::MakeMessage<FTesterKeepAliveMessage>(KeepAliveNumber++);
+		NewKeepAlive->State = State;
+		NewKeepAlive->bReliablySent = true;
+		SendMessage(NewKeepAlive, EMessageBusTesterMessageFlags::Reliable);
+
+		NewKeepAlive = FMessageEndpoint::MakeMessage<FTesterKeepAliveMessage>(KeepAliveNumber++);
+		NewKeepAlive->State = State;
+		NewKeepAlive->bReliablySent = false;
+		SendMessage(NewKeepAlive, EMessageBusTesterMessageFlags::None);
+	}
+
+	//Relaunch ourself if desired
+	if (bShouldSendKeepAlive)
+	{
+		KeepAliveTaskFuture = Async(EAsyncExecution::TaskGraph, [this]()
+		{
+			KeepAliveEvent->Wait(MessageBusTesterUtils::KeepAliveIntervalMilliSeconds);
+			KeepAliveTask();
+		});
+	}
+}
+
+void FMessageBusTester::SetupTestPlan()
+{
+	TestPlanItemTrackers.Empty(TestPlan.TestPlanItems.Num());
+	for (const FTestPlanItem& Item : TestPlan.TestPlanItems)
+	{
+		FTestPlanItemTracker NewTest;
+		NewTest.ItemDescription = Item;
+		TestPlanItemTrackers.Emplace(MoveTemp(NewTest));
+	}
+
+	State = EMessageBusTesterState::Active;
+}
+
+bool FMessageBusTester::SendMessageInternal(FTesterBaseMessage* Payload, UScriptStruct* Type, const TArray<FMessageAddress>& Addresses, EMessageBusTesterMessageFlags InFlags) const
+{
+	if (Addresses.Num() > 0 && TesterEndpoint.IsValid())
+	{
+		//Set identifier in sent message
+		Payload->Identifier = Identifier;
+
+		const EMessageFlags Flags = EnumHasAnyFlags(InFlags, EMessageBusTesterMessageFlags::Reliable) ? EMessageFlags::Reliable : EMessageFlags::None;
+		TesterEndpoint->Send(Payload, Type, Flags, nullptr, Addresses, FTimespan::Zero(), FDateTime::MaxValue());
+		return true;
+	}
+	// We take ownership of Payload so we must release the memory here.  MessageBus uses FMalloc for allocation so we must call Free.
+	FMemory::Free(Payload); 
+	return false;
+}
+
+void FMessageBusTester::HandleDiscoveryMessage(const FTesterDiscoveryMessage& Message, const TSharedRef<IMessageContext, ESPMode::ThreadSafe>& Context)
+{
+	//Don't listen to ourselves
+	if (Message.Identifier == Identifier)
+	{
+		return;
+	}
+
+	TRACE_CPUPROFILER_EVENT_SCOPE(FMessageBusTester::HandleDiscoveryMessage);
+
+	if (Message.Descriptor.MessageVersion != MessageBusTesterUtils::SupportedMessageBusTesterVersion)
+	{
+		//Keep track of invalid providers to avoid log spamming
+		if (!InvalidTesters.Contains(Message.Identifier))
+		{
+			InvalidTesters.Add(Message.Identifier);
+
+			UE_LOG(LogMessageBusTester, Warning, TEXT("Received discovery response from Tester '%s' with unsupported version '%d'. Supported version is '%d'")
+				, *Message.Descriptor.FriendlyName.ToString()
+				, Message.Descriptor.MessageVersion
+				, MessageBusTesterUtils::SupportedMessageBusTesterVersion);
+		}
+
+		return;
+	}
+
+	const UMessageBusTesterSettings* Settings = GetDefault<UMessageBusTesterSettings>();
+	if (!Settings->bUseSessionId || Settings->GetSessionId() == Message.Descriptor.SessionId)
+	{
+		HandleDiscoveredTester(Message.Identifier, Message.Descriptor, Context->GetSender());
+	}
+	else if (!InvalidTesters.Contains(Message.Identifier))
+	{
+		//Keep track of invalid providers to avoid log spamming
+		InvalidTesters.Add(Message.Identifier);
+
+		UE_LOG(LogMessageBusTester, Warning, TEXT("Received discovery response from Tester '%s' with SessionId '%d' but expected '%d'")
+			, *Message.Descriptor.FriendlyName.ToString()
+			, Message.Descriptor.SessionId
+			, Settings->GetSessionId());
+	}
+}
+
+void FMessageBusTester::HandleDiscoveredTester(const FGuid& InIdentifier, const FTesterInstanceDescriptor& Descriptor, const FMessageAddress& Address)
+{
+	// -already known
+	// -known in the passed but cleared
+	// -unknown but matchable to previous instance
+
+	FScopeLock Lock(&DiscoveredTesterCriticalSection);
+	
+	const auto FindTesterLambda = [InIdentifier, &Descriptor](const TSharedPtr<FDiscoveredTester, ESPMode::ThreadSafe>& Other)
+	{
+		//Look for a matching lost/disconnected provider from the same machine + same stage
+		if (Other->Descriptor.MachineName == Descriptor.MachineName && Other->Descriptor.FriendlyName == Descriptor.FriendlyName)
+		{
+			return true;
+		}
+
+		return (Other->Identifier == InIdentifier);
+	};
+
+	const int32 ExistingProviderIndex = DiscoveredTesters.IndexOfByPredicate(FindTesterLambda);
+	if (ExistingProviderIndex != INDEX_NONE)
+	{
+		//Need to better support rediscovering a dropped tester. Especially during an active test plan. 
+		
+		//In case of an existing provider, 
+		TSharedPtr<FDiscoveredTester, ESPMode::ThreadSafe>& Entry = DiscoveredTesters[ExistingProviderIndex];
+		if (Entry->ConnectionState == EDiscoveredTesterConnectionState::Lost)
+		{
+			//Need to update data about a closed tester. Its address will certainly have changed and descriptor could also have
+			const FGuid PreviousIdentifier = Entry->Identifier;
+			FillTesterDescription(InIdentifier, Descriptor, Address, Entry);
+
+			if (InIdentifier != PreviousIdentifier)
+			{
+				//If identifier change, make sure our mapping is up to date
+				ensureAlways(false);
+				//If we're reslotting an old provider to a new one (same stagename / same machine) trigger a list change to let listeners know about it
+				OnDiscoveredTesterListChanged().Broadcast();
+			}
+
+			UE_LOG(LogMessageBusTester, Log, TEXT("Old tester recovered, %s - (%s (%d)) on sessionId %d")
+				, *Descriptor.FriendlyName.ToString()
+				, *Descriptor.MachineName
+				, Descriptor.ProcessId
+				, Descriptor.SessionId);
+		}
+	}
+	else
+	{
+		AddTester(InIdentifier, Descriptor, Address);
+	}
+}
+
+bool FMessageBusTester::AddTester(const FGuid& InIdentifier, const FTesterInstanceDescriptor& Descriptor, const FMessageAddress& Address)
+{
+	FScopeLock Lock(&DiscoveredTesterCriticalSection);
+
+	if (DiscoveredTesters.ContainsByPredicate([InIdentifier](const TSharedPtr<FDiscoveredTester, ESPMode::ThreadSafe>& Other) { return Other->Identifier == InIdentifier; }))
+	{
+		return false;
+	}
+
+	UE_LOG(LogMessageBusTester, Log, TEXT("Adding tester %s - (%s (%d)) on sessionId %d")
+		, *Descriptor.FriendlyName.ToString()
+		, *Descriptor.MachineName
+		, Descriptor.ProcessId
+		, Descriptor.SessionId);
+	
+	TSharedPtr<FDiscoveredTester, ESPMode::ThreadSafe>& NewEntry = DiscoveredTesters.Emplace_GetRef(MakeShared<FDiscoveredTester, ESPMode::ThreadSafe>());
+	FillTesterDescription(InIdentifier, Descriptor, Address, NewEntry);
+
+	//Let know a new tester was added to the list
+	OnDiscoveredTesterListChanged().Broadcast();
+
+	return true;
+}
+
+void FMessageBusTester::FillTesterDescription(const FGuid& InIdentifier, const FTesterInstanceDescriptor& NewDescriptor, const FMessageAddress& NewAddress, TSharedPtr<FDiscoveredTester, ESPMode::ThreadSafe>& OutEntry) const
+{
+	OutEntry->Identifier = InIdentifier;
+	OutEntry->Descriptor = NewDescriptor;
+	OutEntry->Address = NewAddress;
+	OutEntry->State = EMessageBusTesterState::Idle;
+	OutEntry->ConnectionState = EDiscoveredTesterConnectionState::Connected;
+	OutEntry->LastReceivedMessageTime = FPlatformTime::Seconds();
+	OutEntry->LastReliableKeepAliveMessageTime = OutEntry->LastReceivedMessageTime;
+	OutEntry->LastUnreliableKeepAliveMessageTime = OutEntry->LastReceivedMessageTime;
+}
+
+void FMessageBusTester::HandleKeepAlive(const FTesterKeepAliveMessage& Message,	const TSharedRef<IMessageContext, ESPMode::ThreadSafe>& Context)
+{
+	FScopeLock Lock(&DiscoveredTesterCriticalSection);
+
+	const FGuid& SenderIdentifier = Message.Identifier;
+	if (TSharedPtr<FDiscoveredTester, ESPMode::ThreadSafe>* FoundTesterPtr = DiscoveredTesters.FindByPredicate([SenderIdentifier](const TSharedPtr<FDiscoveredTester, ESPMode::ThreadSafe>& Other) { return Other->Identifier == SenderIdentifier; }))
+	{
+		TSharedPtr<FDiscoveredTester, ESPMode::ThreadSafe>& FoundTester = *FoundTesterPtr;
+		if (FoundTester->ConnectionState != EDiscoveredTesterConnectionState::Lost)
+		{
+			FoundTester->NotifyKeepAliveReceived(Message.bReliablySent);
+			FoundTester->State = Message.State;
+		}
+	}
+}
+
+void FMessageBusTester::HandleStartTest(const FTesterStartTestMessage& Message,	const TSharedRef<IMessageContext, ESPMode::ThreadSafe>& Context)
+{
+	TestPlan = Message.TestPlan;
+	OnTestPlanChanged().Broadcast();
+	SetupTestPlan();
+}
+
+void FMessageBusTester::HandleStopTest(const FTesterStopTestMessage& Message, const TSharedRef<IMessageContext, ESPMode::ThreadSafe>& Context)
+{
+	State = EMessageBusTesterState::Idle;
+}
+
+void FMessageBusTester::HandlePayloadMessage(const FTestPayloadMessage& Message, const TSharedRef<IMessageContext, ESPMode::ThreadSafe>& Context)
+{
+	FTestPayloadReceptionMessage* Response = FMessageEndpoint::MakeMessage<FTestPayloadReceptionMessage>();
+	Response->PayloadId = Message.PayloadId;
+	Response->TestPlanIndex = Message.TestPlanIndex;
+	Response->ReceivedPayloadSize = Message.Payload.Num();
+	SendMessage(Response,  { Context->GetSender() }, EMessageBusTesterMessageFlags::Reliable);
+}
+
+void FMessageBusTester::HandlePayloadReceptionMessage(const FTestPayloadReceptionMessage& Message, const TSharedRef<IMessageContext, ESPMode::ThreadSafe>& Context)
+{
+	if (TestPlanItemTrackers.IsValidIndex(Message.TestPlanIndex))
+	{
+		const FGuid& SenderIdentifier = Message.Identifier;
+		if (TSharedPtr<FDiscoveredTester, ESPMode::ThreadSafe>* FoundTesterPtr = DiscoveredTesters.FindByPredicate([SenderIdentifier](const TSharedPtr<FDiscoveredTester, ESPMode::ThreadSafe>& Other) { return Other->Identifier == SenderIdentifier; }))
+		{
+			(*FoundTesterPtr)->NotifyAckReceived(Message);
+		}
+	}
+}
+
diff --git a/Engine/Plugins/Experimental/MessageBusTester/Source/MessageBusTester/Private/MessageBusTester.h b/Engine/Plugins/Experimental/MessageBusTester/Source/MessageBusTester/Private/MessageBusTester.h
new file mode 100644
--- /dev/null
+++ b/Engine/Plugins/Experimental/MessageBusTester/Source/MessageBusTester/Private/MessageBusTester.h
@@ -0,0 +1,153 @@
+// Copyright Epic Games, Inc. All Rights Reserved.
+
+#pragma once
+
+#include "IMessageBusTester.h"
+
+#include "Containers/Ticker.h"
+
+#include "DiscoveredTester.h"
+#include "IMessageContext.h"
+#include "UObject/ObjectMacros.h"
+
+#include <atomic>
+#include "Async/Future.h"
+
+class FMessageEndpoint;
+
+/**
+ * Implementation of the MessageBus Tester
+ */
+class FMessageBusTester : public IMessageBusTester
+{
+public:
+
+	FMessageBusTester() = default;
+	FMessageBusTester(const FMessageBusTester&) = delete;
+	FMessageBusTester& operator=(const FMessageBusTester&) = delete;
+	virtual ~FMessageBusTester();
+
+	//~Begin IMessageBusTester interface
+	virtual bool IsActive() const override { return State == EMessageBusTesterState::Active; }
+	virtual FOnDiscoveredTesterListChanged& OnDiscoveredTesterListChanged() override { return DiscoveredTesterListChangedDelegate; }
+	virtual FOnTestPlanChanged& OnTestPlanChanged() override { return TestPlanChangedDelegate; }
+	virtual TConstArrayView<TSharedPtr<FDiscoveredTester, ESPMode::ThreadSafe>> GetDiscoveredTesters() const override { return MakeArrayView(DiscoveredTesters); }
+	virtual const FMessageBusTestPlan& GetTestPlan() const override { return TestPlan; }
+	virtual void AddTestPlanItem(FTestPlanItem NewItem) override;
+	virtual void RemoveTestPlanItem(int32 Index) override;
+	virtual bool StartTest() override;
+	virtual bool StopTest() override;
+	virtual EMessageBusTesterState GetState() const override { return State; }
+	//~End IMessageBusTester
+
+
+	void Initialize();
+	void Start();
+	void Stop();
+
+private:
+
+	/** Sends a message using messagebus to connected providers */
+	bool SendMessageInternal(FTesterBaseMessage* Payload, UScriptStruct* Type, const TArray<FMessageAddress>& Addresses, EMessageBusTesterMessageFlags InFlags) const;
+
+	void HandleDiscoveryMessage(const FTesterDiscoveryMessage& Message, const TSharedRef<IMessageContext, ESPMode::ThreadSafe>& Context);
+	void HandleDiscoveredTester(const FGuid& Identifier, const FTesterInstanceDescriptor& Descriptor,
+	                            const FMessageAddress& Address);
+	bool AddTester(const FGuid& Identifier, const FTesterInstanceDescriptor& Descriptor,
+	               const FMessageAddress& Address);
+	void FillTesterDescription(const FGuid& Identifier, const FTesterInstanceDescriptor& NewDescriptor,
+	                           const FMessageAddress& NewAddress, TSharedPtr<FDiscoveredTester, ESPMode::ThreadSafe>& OutEntry) const;
+	void HandleKeepAlive(const FTesterKeepAliveMessage& Message, const TSharedRef<IMessageContext, ESPMode::ThreadSafe>& Context);
+	void HandleStartTest(const FTesterStartTestMessage& Message, const TSharedRef<IMessageContext, ESPMode::ThreadSafe>& Context);
+	void HandleStopTest(const FTesterStopTestMessage& Message, const TSharedRef<IMessageContext, ESPMode::ThreadSafe>& Context);
+	void HandlePayloadMessage(const FTestPayloadMessage& Message, const TSharedRef<IMessageContext, ESPMode::ThreadSafe>& Context);
+	void HandlePayloadReceptionMessage(const FTestPayloadReceptionMessage& Message, const TSharedRef<IMessageContext, ESPMode::ThreadSafe>& Context);
+	
+	/**  */
+	bool Tick(float DeltaTime);
+
+	/** Updates each provider's state based on last communication timestamp */
+	void UpdateDiscoveredTester();
+
+	/** Broadcasts a discovery message to find new providers and also keep communication active with all providers */
+	void SendDiscoveryMessage();
+
+	/** Shutdowns monitor before everything has exited. Let know providers that we're out */
+	void OnPreExit();
+
+	void UpdateTestSequence();
+
+	/** Generic send message to support constructor parameters directly and temp object created */
+	template<typename MessageType>
+	bool SendMessage(MessageType* Payload, EMessageBusTesterMessageFlags Flags)
+	{
+		static_assert(TIsDerivedFrom<MessageType, FTesterBaseMessage>::IsDerived, "MessageType must be a FTesterBaseMessage derived UStruct.");
+
+		FScopeLock Lock(&DiscoveredTesterCriticalSection);
+
+		TArray<FMessageAddress> TestersAddress;
+		TestersAddress.Reserve(DiscoveredTesters.Num());
+		for (const TSharedPtr<FDiscoveredTester, ESPMode::ThreadSafe>& Tester : DiscoveredTesters)
+		{
+			if (Tester->ConnectionState == EDiscoveredTesterConnectionState::Connected)
+			{
+				TestersAddress.Add(Tester->Address);
+			}
+		}
+		return SendMessageInternal(Payload, MessageType::StaticStruct(), TestersAddress, Flags);
+	}
+	
+	/** Generic send message to support constructor parameters directly and temp object created */
+	template<typename MessageType>
+	bool SendMessage(MessageType* Payload, const TArray<FMessageAddress>& Addresses, EMessageBusTesterMessageFlags Flags)
+	{
+		static_assert(TIsDerivedFrom<MessageType, FTesterBaseMessage>::IsDerived, "MessageType must be a FTesterBaseMessage derived UStruct.");
+		return SendMessageInternal(Payload, MessageType::StaticStruct(), Addresses, Flags);
+	}
+
+	void KeepAliveTask();
+
+	void SetupTestPlan();
+
+private:
+
+	/** Endpoint used to communicate with other testers */
+	TSharedPtr<FMessageEndpoint, ESPMode::ThreadSafe> TesterEndpoint;
+
+	TArray<TSharedPtr<FDiscoveredTester, ESPMode::ThreadSafe>> DiscoveredTesters;
+
+	/** List of providers we received discovery response from that were not usable (wrong version, session id...) */
+	TArray<FGuid> InvalidTesters;
+
+	/** Timestamp when we last sent a discovery message */
+	double LastSentDiscoveryMessage = 0.0;
+
+	/** Handle to our ticking delegate. */
+	FTSTicker::FDelegateHandle TickerHandle;
+
+	/** This tester identifier */
+	FGuid Identifier;
+
+	/** DiscoveryMessage built once and broadcasted periodically */
+	FTesterDiscoveryMessage CachedDiscoveryMessage;
+
+	EMessageBusTesterState State = EMessageBusTesterState::Idle;
+
+	FCriticalSection DiscoveredTesterCriticalSection;
+	std::atomic<bool> bShouldSendKeepAlive;
+	TFuture<void> KeepAliveTaskFuture;
+	FEvent* KeepAliveEvent = nullptr;
+
+	FOnDiscoveredTesterListChanged DiscoveredTesterListChangedDelegate;
+	FOnTestPlanChanged TestPlanChangedDelegate;
+
+	FMessageBusTestPlan TestPlan;
+
+	struct FTestPlanItemTracker
+	{
+		FTestPlanItem ItemDescription;
+		TMap<FGuid, int32> SentPayloadIds;
+		double LastPayloadSentTimestamp = 0.0;
+	};
+	TArray<FTestPlanItemTracker> TestPlanItemTrackers;
+};
diff --git a/Engine/Plugins/Experimental/MessageBusTester/Source/MessageBusTester/Private/MessageBusTesterCommon.cpp b/Engine/Plugins/Experimental/MessageBusTester/Source/MessageBusTester/Private/MessageBusTesterCommon.cpp
new file mode 100644
--- /dev/null
+++ b/Engine/Plugins/Experimental/MessageBusTester/Source/MessageBusTester/Private/MessageBusTesterCommon.cpp
@@ -0,0 +1,35 @@
+// Copyright Epic Games, Inc. All Rights Reserved.
+
+#include "MessageBusTesterCommon.h"
+
+#include "Misc/App.h"
+#include "MessageBusTesterSettings.h"
+
+
+namespace MessageBusTesterUtils
+{
+	static const FString CachedComputerName = FPlatformProcess::ComputerName();
+	static const uint32 CachedProcessId = FPlatformProcess::GetCurrentProcessId();
+
+	FTesterInstanceDescriptor GetInstanceDescriptor()
+	{
+		FTesterInstanceDescriptor Descriptor;
+
+		//A machine could spawn multiple UE instances. Need to be able to differentiate them. ProcessId is there for that reason
+		Descriptor.MachineName = CachedComputerName;
+		Descriptor.ProcessId = CachedProcessId;
+
+		Descriptor.FriendlyName = GetDefault<UMessageBusTesterSettings>()->CommandLineFriendlyName;
+		if (Descriptor.FriendlyName == NAME_None)
+		{
+			const FString TempName = FString::Printf(TEXT("%s:%u"), *Descriptor.MachineName, Descriptor.ProcessId);
+			Descriptor.FriendlyName = *TempName;
+		}
+
+		Descriptor.SessionId = GetDefault<UMessageBusTesterSettings>()->GetSessionId();
+
+		return Descriptor;
+	}
+}
+
+
diff --git a/Engine/Plugins/Experimental/MessageBusTester/Source/MessageBusTester/Private/MessageBusTesterLogger.cpp b/Engine/Plugins/Experimental/MessageBusTester/Source/MessageBusTester/Private/MessageBusTesterLogger.cpp
new file mode 100644
--- /dev/null
+++ b/Engine/Plugins/Experimental/MessageBusTester/Source/MessageBusTester/Private/MessageBusTesterLogger.cpp
@@ -0,0 +1,17 @@
+// Copyright Epic Games, Inc. All Rights Reserved.
+
+#include "MessageBusTesterLogger.h"
+
+
+void FMessageBusTesterLogger::Log(const FName& TesterName, FString&& InLogMessage, EMessageSeverity::Type Verbosity /* = EMessageSeverity::Info */)
+{
+	TSharedPtr<FMessageBusTesterLogEntry> NewEntry = MakeShared<FMessageBusTesterLogEntry>(TesterName, MoveTemp(InLogMessage), Verbosity);
+	Entries.Add(NewEntry);
+	OnMessageBusTesterNewLogReceived().Broadcast(NewEntry.ToSharedRef());
+}
+
+void FMessageBusTesterLogger::ClearLog()
+{
+	Entries.Empty();
+	OnMessageBusTesterLogCleared().Broadcast();
+}
diff --git a/Engine/Plugins/Experimental/MessageBusTester/Source/MessageBusTester/Private/MessageBusTesterLogger.h b/Engine/Plugins/Experimental/MessageBusTester/Source/MessageBusTester/Private/MessageBusTesterLogger.h
new file mode 100644
--- /dev/null
+++ b/Engine/Plugins/Experimental/MessageBusTester/Source/MessageBusTester/Private/MessageBusTesterLogger.h
@@ -0,0 +1,31 @@
+// Copyright Epic Games, Inc. All Rights Reserved.
+
+#pragma once
+
+#include "IMessageBusTesterLogger.h"
+
+
+
+/**
+ *  
+ */
+class FMessageBusTesterLogger : public IMessageBusTesterLogger
+{
+public:
+
+
+	//~Begin IMessageBusTesterLogger interface
+	virtual void Log(const FName& TesterName, FString&& InLogMessage, EMessageSeverity::Type Verbosity /* = EMessageSeverity::Info */) override;
+	virtual void ClearLog() override;
+	virtual FOnMessageBusTesterNewLogReceived& OnMessageBusTesterNewLogReceived() override { return NewLogReceivedDelegate; }
+	virtual FOnMessageBusTesterLogCleared& OnMessageBusTesterLogCleared() override { return LogClearedDelegate; }
+	//~End IMessageBusTesterLogger
+
+
+private:
+
+	TArray<TSharedPtr<FMessageBusTesterLogEntry>> Entries;
+
+	FOnMessageBusTesterNewLogReceived NewLogReceivedDelegate;
+	FOnMessageBusTesterLogCleared LogClearedDelegate;
+};
diff --git a/Engine/Plugins/Experimental/MessageBusTester/Source/MessageBusTester/Private/MessageBusTesterModule.cpp b/Engine/Plugins/Experimental/MessageBusTester/Source/MessageBusTester/Private/MessageBusTesterModule.cpp
new file mode 100644
--- /dev/null
+++ b/Engine/Plugins/Experimental/MessageBusTester/Source/MessageBusTester/Private/MessageBusTesterModule.cpp
@@ -0,0 +1,44 @@
+// Copyright Epic Games, Inc. All Rights Reserved.
+
+#include "MessageBusTesterModule.h"
+
+#include "HAL/IConsoleManager.h"
+#include "IMessageBusTesterLogger.h"
+#include "Misc/CoreDelegates.h"
+#include "MessageBusTester.h"
+#include "MessageBusTesterLogger.h"
+
+const FName IMessageBusTesterModule::ModuleName = TEXT("MessageBusTester");
+
+
+DEFINE_LOG_CATEGORY(LogMessageBusTester)
+
+void FMessageBusTesterModule::StartupModule()
+{
+	MessageBusTester = MakeUnique<FMessageBusTester>();
+	MessageBusTester->Initialize();
+
+	Logger = MakeUnique<FMessageBusTesterLogger>();
+}
+
+void FMessageBusTesterModule::ShutdownModule()
+{
+	MessageBusTester.Reset();
+}
+
+IMessageBusTester& FMessageBusTesterModule::GetMessageBusTester()
+{
+	check(MessageBusTester.IsValid());
+	return *MessageBusTester;
+}
+
+IMessageBusTesterLogger& FMessageBusTesterModule::GetLogger()
+{
+	check(Logger.IsValid());
+	return *Logger;
+}
+
+IMPLEMENT_MODULE(FMessageBusTesterModule, MessageBusTester)
+
+
+
diff --git a/Engine/Plugins/Experimental/MessageBusTester/Source/MessageBusTester/Private/MessageBusTesterModule.h b/Engine/Plugins/Experimental/MessageBusTester/Source/MessageBusTester/Private/MessageBusTesterModule.h
new file mode 100644
--- /dev/null
+++ b/Engine/Plugins/Experimental/MessageBusTester/Source/MessageBusTester/Private/MessageBusTesterModule.h
@@ -0,0 +1,41 @@
+// Copyright Epic Games, Inc. All Rights Reserved.
+
+#pragma once
+
+#include "IMessageBusTesterModule.h"
+
+#include "CoreMinimal.h"
+#include "Modules/ModuleManager.h"
+#include "Templates/UniquePtr.h"
+
+class FAutoConsoleCommand;
+class FMessageBusTester;
+class FMessageBusTesterLogger;
+
+DECLARE_LOG_CATEGORY_EXTERN(LogMessageBusTester, Log, All);
+
+
+class FMessageBusTesterModule : public IMessageBusTesterModule
+{
+public:	
+	virtual ~FMessageBusTesterModule() = default;
+	
+	//~ Begin IModuleInterface interface
+	virtual void StartupModule() override;
+	virtual void ShutdownModule() override;
+	//~End IModuleInterface
+
+	//~Begin IMessageBusTesterModule interface
+	 virtual IMessageBusTester& GetMessageBusTester() override;
+	 virtual IMessageBusTesterLogger& GetLogger() override;
+	//~End IMessageBusTesterModule interface
+
+private:
+
+
+protected:
+	
+	/** Single instance of the MessageBusTester */
+	TUniquePtr<FMessageBusTester> MessageBusTester;
+	TUniquePtr<FMessageBusTesterLogger> Logger;
+};
diff --git a/Engine/Plugins/Experimental/MessageBusTester/Source/MessageBusTester/Private/MessageBusTesterSettings.cpp b/Engine/Plugins/Experimental/MessageBusTester/Source/MessageBusTester/Private/MessageBusTesterSettings.cpp
new file mode 100644
--- /dev/null
+++ b/Engine/Plugins/Experimental/MessageBusTester/Source/MessageBusTester/Private/MessageBusTesterSettings.cpp
@@ -0,0 +1,40 @@
+// Copyright Epic Games, Inc. All Rights Reserved.
+
+#include "MessageBusTesterSettings.h"
+
+#include "Misc/CommandLine.h"
+
+UMessageBusTesterSettings::UMessageBusTesterSettings()
+{
+	int32 CommandLineSessionIdTmp = 0;
+	if (FParse::Value(FCommandLine::Get(), TEXT("-MessageBusTesterSessionId="), CommandLineSessionIdTmp))
+	{
+		CommandLineSessionId = CommandLineSessionIdTmp;
+	}
+
+	FParse::Value(FCommandLine::Get(), TEXT("-MessageBusTesterFriendlyName="), CommandLineFriendlyName);
+}
+
+FName UMessageBusTesterSettings::GetCategoryName() const
+{
+	return TEXT("Plugins");
+}
+
+#if WITH_EDITOR
+FText UMessageBusTesterSettings::GetSectionText() const
+{
+	return NSLOCTEXT("MessageBusTesteringPlugin", "MessageBusTesterSettingsSection", "MessageBus Tester");
+}
+#endif
+
+
+int32 UMessageBusTesterSettings::GetSessionId() const
+{
+	if (CommandLineSessionId.IsSet())
+	{
+		return CommandLineSessionId.GetValue();
+	}
+
+	return SessionId;
+}
+
diff --git a/Engine/Plugins/Experimental/MessageBusTester/Source/MessageBusTester/Public/DiscoveredTester.h b/Engine/Plugins/Experimental/MessageBusTester/Source/MessageBusTester/Public/DiscoveredTester.h
new file mode 100644
--- /dev/null
+++ b/Engine/Plugins/Experimental/MessageBusTester/Source/MessageBusTester/Public/DiscoveredTester.h
@@ -0,0 +1,100 @@
+// Copyright Epic Games, Inc. All Rights Reserved.
+
+#pragma once
+
+
+#include "CoreMinimal.h"
+#include "Misc/QualifiedFrameTime.h"
+#include "UObject/ObjectMacros.h"
+#include "IMessageContext.h"
+#include "MessageBusTesterCommon.h"
+#include "Math/MovingWindowAverageFast.h"
+#include "INetworkMessagingExtension.h"
+
+
+
+struct FPayloadInfo
+{
+	FPayloadInfo(const FTestPayloadMessage& PayloadMessage)
+		: PayloadId(PayloadMessage.PayloadId)
+		, SentTime(FPlatformTime::Seconds())
+		, PayloadSize(PayloadMessage.Payload.Num())
+	{
+		
+	}
+
+	bool operator==(const FPayloadInfo& Other) const
+	{
+		return PayloadId == Other.PayloadId && PayloadSize == Other.PayloadSize;
+	}
+	
+	int32 PayloadId = INDEX_NONE;
+	double SentTime = 0.0;
+	int64 PayloadSize = 0;
+	double AcknowledgedTime = -1.0;
+};
+
+struct FKeepAliveStatistics
+{
+	double LastUnreliableKeepAliveReceptionInterval = 0.0;
+	double MinUnreliableKeepAliveInterval = DBL_MAX;
+	double MaxUnreliableKeepAliveInterval = 0.0;
+	double AverageUnreliableKeepAliveInterval = 0.0;
+	double LastReliableKeepAliveReceptionInterval = 0.0;
+	double MinReliableKeepAliveInterval = DBL_MAX;
+	double MaxReliableKeepAliveInterval = 0.0;
+	double AverageReliableKeepAliveInterval = 0.0;
+};
+
+
+/**
+*/
+class MESSAGEBUSTESTER_API FDiscoveredTester
+{
+	bool operator==(const FDiscoveredTester& Other) const { return Identifier == Other.Identifier; }
+	bool operator!=(const FDiscoveredTester& Other) const { return !(*this == Other); }
+
+
+public:
+	void UpdateConnectionState();
+	void NotifyPayloadSent(const FTestPayloadMessage& PayloadMessage);
+	void NotifyAckReceived(const FTestPayloadReceptionMessage& AckMessage);
+	void NotifyKeepAliveReceived(bool bIsReliable);
+	bool IsPendingPayloadAck(int32 PayloadId) const;
+
+public:
+
+	/** Identifier of this tester */
+	FGuid Identifier;
+
+	/** Detailed descriptor */
+	FTesterInstanceDescriptor Descriptor;
+
+	/** Address of this Tester */
+	FMessageAddress Address;
+
+	/** State of this Tester based on message reception */
+	EMessageBusTesterState State = EMessageBusTesterState::Idle;
+
+	/** State of this Tester based on message reception */
+	EDiscoveredTesterConnectionState ConnectionState = EDiscoveredTesterConnectionState::Connected;
+
+	FMessageTransportStatistics Statistics;
+
+	/** Timestamp when last message was received based on FPlatformTime::Seconds */
+	double LastReceivedMessageTime = 0.0;
+
+	/** Timestamp when last KeepAlive was received based on FPlatformTime::Seconds */
+	double LastUnreliableKeepAliveMessageTime = 0.0;
+
+	/** Timestamp when last KeepAlive was received based on FPlatformTime::Seconds */
+	double LastReliableKeepAliveMessageTime = 0.0;
+
+	/** Last interval between keep alive reception */
+	FMovingWindowAverageFast<double, 128> AverageUnreliableKeepAliveIntervalWindow;
+	FMovingWindowAverageFast<double, 128> AverageReliableKeepAliveIntervalWindow;
+	FKeepAliveStatistics KeepAliveStatistics;
+
+	TArray<FPayloadInfo> SentPayloads;
+	TArray<FPayloadInfo> AcknowledgedPayloads;
+};
diff --git a/Engine/Plugins/Experimental/MessageBusTester/Source/MessageBusTester/Public/IMessageBusTester.h b/Engine/Plugins/Experimental/MessageBusTester/Source/MessageBusTester/Public/IMessageBusTester.h
new file mode 100644
--- /dev/null
+++ b/Engine/Plugins/Experimental/MessageBusTester/Source/MessageBusTester/Public/IMessageBusTester.h
@@ -0,0 +1,38 @@
+// Copyright Epic Games, Inc. All Rights Reserved.
+
+#pragma once
+
+#include "CoreMinimal.h"
+
+#include "MessageBusTesterCommon.h"
+
+class FDiscoveredTester;
+
+
+/**
+ * Interface for the MessageBusTester */
+class MESSAGEBUSTESTER_API IMessageBusTester
+{
+public:
+
+	virtual ~IMessageBusTester() {}
+	
+	/**
+	 * Returns true if tester is active
+	 */
+	virtual bool IsActive() const = 0;
+
+	virtual TConstArrayView<TSharedPtr<FDiscoveredTester, ESPMode::ThreadSafe>> GetDiscoveredTesters() const = 0;
+	virtual const FMessageBusTestPlan& GetTestPlan() const = 0;
+	virtual bool StartTest() = 0;
+	virtual bool StopTest() = 0;
+	virtual void AddTestPlanItem(FTestPlanItem NewItem) = 0;
+	virtual void RemoveTestPlanItem(int32 Index) = 0;
+	virtual EMessageBusTesterState GetState() const = 0;
+
+	DECLARE_MULTICAST_DELEGATE(FOnDiscoveredTesterListChanged);
+	virtual FOnDiscoveredTesterListChanged& OnDiscoveredTesterListChanged() = 0;
+
+	DECLARE_MULTICAST_DELEGATE(FOnTestPlanChanged);
+	virtual FOnTestPlanChanged& OnTestPlanChanged() = 0;
+};
diff --git a/Engine/Plugins/Experimental/MessageBusTester/Source/MessageBusTester/Public/IMessageBusTesterLogger.h b/Engine/Plugins/Experimental/MessageBusTester/Source/MessageBusTester/Public/IMessageBusTesterLogger.h
new file mode 100644
--- /dev/null
+++ b/Engine/Plugins/Experimental/MessageBusTester/Source/MessageBusTester/Public/IMessageBusTesterLogger.h
@@ -0,0 +1,39 @@
+// Copyright Epic Games, Inc. All Rights Reserved.
+
+#pragma once
+
+#include "CoreMinimal.h"
+
+#include "Logging/TokenizedMessage.h"
+
+
+struct FMessageBusTesterLogEntry
+{
+	FMessageBusTesterLogEntry() = default;
+	FMessageBusTesterLogEntry(FName InLogSource, FString&& InLogMessage, EMessageSeverity::Type InMessageSeverity)
+		: Source(InLogSource)
+		, LogMessage(MoveTemp(InLogMessage))
+		, MessageSeverity(InMessageSeverity)
+	{}
+
+	FName Source = NAME_None;
+	FString LogMessage;
+	EMessageSeverity::Type MessageSeverity = EMessageSeverity::Info;
+};
+
+
+class MESSAGEBUSTESTER_API IMessageBusTesterLogger
+{
+public:
+
+	virtual ~IMessageBusTesterLogger() = default;
+
+	virtual void Log(const FName& TesterName, FString&& InLogMessage, EMessageSeverity::Type Verbosity = EMessageSeverity::Info) = 0;
+	virtual void ClearLog() = 0;
+
+	DECLARE_MULTICAST_DELEGATE_OneParam(FOnMessageBusTesterNewLogReceived, TSharedRef<FMessageBusTesterLogEntry> /*LogData*/);
+	virtual FOnMessageBusTesterNewLogReceived& OnMessageBusTesterNewLogReceived() = 0;
+
+	DECLARE_MULTICAST_DELEGATE(FOnMessageBusTesterLogCleared);
+	virtual FOnMessageBusTesterLogCleared& OnMessageBusTesterLogCleared() = 0;
+};
diff --git a/Engine/Plugins/Experimental/MessageBusTester/Source/MessageBusTester/Public/IMessageBusTesterModule.h b/Engine/Plugins/Experimental/MessageBusTester/Source/MessageBusTester/Public/IMessageBusTesterModule.h
new file mode 100644
--- /dev/null
+++ b/Engine/Plugins/Experimental/MessageBusTester/Source/MessageBusTester/Public/IMessageBusTesterModule.h
@@ -0,0 +1,47 @@
+// Copyright Epic Games, Inc. All Rights Reserved.
+
+#pragma once
+
+#include "Modules/ModuleInterface.h"
+
+#include "CoreMinimal.h"
+#include "Modules/ModuleManager.h"
+
+class IMessageBusTester;
+class IMessageBusTesterLogger;
+
+class MESSAGEBUSTESTER_API IMessageBusTesterModule : public IModuleInterface
+{
+public:	
+	/**
+	 * Singleton-like access to this module's interface.  This is just for convenience!
+	 * Beware of calling this during the shutdown phase, though.  Your module might have been unloaded already.
+	 *
+	 * @return Returns singleton instance, loading the module on demand if needed
+	 */
+	static inline IMessageBusTesterModule& Get()
+	{
+		return FModuleManager::LoadModuleChecked<IMessageBusTesterModule>(ModuleName);
+	}
+
+	/**
+	 * Checks to see if this module is loaded and ready.  It is only valid to call Get() during shutdown if IsAvailable() returns true.
+	 *
+	 * @return True if the module is loaded and ready to use
+	 */
+	static inline bool IsAvailable()
+	{
+		return FModuleManager::Get().IsModuleLoaded(ModuleName);
+	}
+
+	/**
+	 * Returns the MessageBus tester instance
+	 */
+	virtual IMessageBusTester& GetMessageBusTester() = 0;
+
+	virtual IMessageBusTesterLogger& GetLogger() = 0;
+
+private:
+	
+	static const FName ModuleName;
+};
diff --git a/Engine/Plugins/Experimental/MessageBusTester/Source/MessageBusTester/Public/MessageBusTesterCommon.h b/Engine/Plugins/Experimental/MessageBusTester/Source/MessageBusTester/Public/MessageBusTesterCommon.h
new file mode 100644
--- /dev/null
+++ b/Engine/Plugins/Experimental/MessageBusTester/Source/MessageBusTester/Public/MessageBusTesterCommon.h
@@ -0,0 +1,250 @@
+// Copyright Epic Games, Inc. All Rights Reserved.
+
+#pragma once
+
+
+#include "CoreMinimal.h"
+#include "Misc/QualifiedFrameTime.h"
+#include "UObject/ObjectMacros.h"
+#include "IMessageContext.h"
+
+#include "MessageBusTesterCommon.generated.h"
+
+
+/** Message flags configuring how a message is sent */
+UENUM()
+enum class EMessageBusTesterMessageFlags : uint8
+{
+	None = 0,
+
+	/** Sends this message as reliable, to make sure it's received by the receivers */
+	Reliable = 1 << 0,
+};
+ENUM_CLASS_FLAGS(EMessageBusTesterMessageFlags);
+
+/** States that a tester can be in */
+UENUM()
+enum class EMessageBusTesterState : uint8
+{
+	/** Testing is active (messages are sent between each testers) */
+	Active,
+	/** Tester is idle */
+	Idle,
+};
+
+/** Connection state for a MessageBus tester. */
+UENUM()
+enum class EDiscoveredTesterConnectionState : uint8
+{
+	/** Tester is considered connected */
+	Connected,
+	/** Tester was considered lost */
+	Lost,
+};
+
+/**
+ * Holds descriptive information about another tester. 
+ * Information that won't change for a session
+ * Used when testers connects
+ */
+USTRUCT()
+struct MESSAGEBUSTESTER_API FTesterInstanceDescriptor
+{
+	GENERATED_BODY()
+
+public:
+
+public:
+
+	/** Machine name read from FPlatformProcess::ComputerName() */
+	UPROPERTY()
+	FString MachineName;
+
+	/** ProcessId read from FPlatformProcess::GetCurrentProcessId */
+	UPROPERTY()
+	uint32 ProcessId = 0;
+
+	/** Friendly name for this Unreal instance. If empty, this will be MachineName - ProcessId */
+	UPROPERTY()
+	FName FriendlyName;
+
+	/** Session Id that may be used to differentiate different sessions on the network */
+	UPROPERTY()
+	int32 SessionId = INDEX_NONE;
+	
+	/** Provision for versioning a tester if we need to differentiate version. */
+	UPROPERTY()
+	int32 MessageVersion = 1;
+};
+
+/**
+ *  Base structure for all message bus tester messages
+ */
+USTRUCT()
+struct MESSAGEBUSTESTER_API FTesterBaseMessage
+{
+	GENERATED_BODY()
+
+public:
+
+	/** Identifier of this instance */
+	UPROPERTY()
+	FGuid Identifier;
+
+	/** FrameTime of the sender. */
+	UPROPERTY()
+	FQualifiedFrameTime FrameTime;
+};
+
+/**
+ *  Message broadcasted periodically by the monitor to discover new providers 
+ */
+USTRUCT()
+struct MESSAGEBUSTESTER_API FTesterDiscoveryMessage : public FTesterBaseMessage
+{
+	GENERATED_BODY()
+
+public:
+	FTesterDiscoveryMessage() = default;
+	FTesterDiscoveryMessage(FTesterInstanceDescriptor&& InDescriptor)
+		: Descriptor(MoveTemp(InDescriptor))
+	{}
+
+public:
+
+	/** Detailed description of a tester */
+	UPROPERTY()
+	FTesterInstanceDescriptor Descriptor;
+};
+
+/**
+ * 
+ */
+USTRUCT()
+struct MESSAGEBUSTESTER_API FTesterKeepAliveMessage : public FTesterBaseMessage
+{
+	GENERATED_BODY()
+
+public:
+	FTesterKeepAliveMessage() = default;
+	FTesterKeepAliveMessage(uint64 InNumber)
+		: KeepAliveNumber(InNumber)
+	{}
+
+
+public:
+
+	UPROPERTY()
+	uint64 KeepAliveNumber = 0;
+
+	UPROPERTY()
+	bool bReliablySent = true;
+
+	UPROPERTY()
+	EMessageBusTesterState State = EMessageBusTesterState::Idle;
+};
+
+USTRUCT()
+struct FTestPlanItem
+{
+	GENERATED_BODY()
+
+	UPROPERTY()
+	int32 NumBytes = 1024;
+
+	UPROPERTY()
+	float IntervalSeconds = 0;
+};
+
+USTRUCT()
+struct FMessageBusTestPlan
+{
+	GENERATED_BODY()
+
+	UPROPERTY()
+	TArray<FTestPlanItem> TestPlanItems;
+};
+
+/**
+ *
+ */
+USTRUCT()
+struct MESSAGEBUSTESTER_API FTesterStartTestMessage : public FTesterBaseMessage
+{
+	GENERATED_BODY()
+
+public:
+	FTesterStartTestMessage() = default;
+
+
+public:
+	UPROPERTY()
+	FMessageBusTestPlan TestPlan;
+};
+
+/**
+ *
+ */
+USTRUCT()
+struct MESSAGEBUSTESTER_API FTesterStopTestMessage : public FTesterBaseMessage
+{
+	GENERATED_BODY()
+
+public:
+	FTesterStopTestMessage() = default;
+
+public:
+};
+
+
+
+/**
+ *
+ */
+USTRUCT()
+struct MESSAGEBUSTESTER_API FTestPayloadMessage : public FTesterBaseMessage
+{
+	GENERATED_BODY()
+
+public:
+	FTestPayloadMessage() = default;
+
+public:
+	
+	UPROPERTY()
+	int32 TestPlanIndex = INDEX_NONE;
+
+	UPROPERTY()
+	int32 PayloadId = 0;
+
+	UPROPERTY()
+	TArray<uint8> Payload;
+};
+
+/**
+ *
+ */
+USTRUCT()
+struct MESSAGEBUSTESTER_API FTestPayloadReceptionMessage : public FTesterBaseMessage
+{
+	GENERATED_BODY()
+
+public:
+	FTestPayloadReceptionMessage() = default;
+
+public:
+
+	UPROPERTY()
+	int32 TestPlanIndex = INDEX_NONE;
+
+	UPROPERTY()
+	int32 PayloadId = 0;
+
+	UPROPERTY()
+	int32 ReceivedPayloadSize = 0;
+};
+
+namespace MessageBusTesterUtils
+{
+	MESSAGEBUSTESTER_API FTesterInstanceDescriptor GetInstanceDescriptor();
+}
diff --git a/Engine/Plugins/Experimental/MessageBusTester/Source/MessageBusTester/Public/MessageBusTesterSettings.h b/Engine/Plugins/Experimental/MessageBusTester/Source/MessageBusTester/Public/MessageBusTesterSettings.h
new file mode 100644
--- /dev/null
+++ b/Engine/Plugins/Experimental/MessageBusTester/Source/MessageBusTester/Public/MessageBusTesterSettings.h
@@ -0,0 +1,67 @@
+// Copyright Epic Games, Inc. All Rights Reserved.
+
+#pragma once
+
+#include "Engine/DeveloperSettings.h"
+
+#include "UObject/Object.h"
+
+#include "MessageBusTesterSettings.generated.h"
+
+/**
+ * Settings for the MessageBusTester plugin modules. 
+ */
+UCLASS(config=Engine)
+class MESSAGEBUSTESTER_API UMessageBusTesterSettings : public UDeveloperSettings
+{
+	GENERATED_BODY()
+
+public:
+	UMessageBusTesterSettings();
+	//~ Begin UDevelopperSettings interface
+	virtual FName GetCategoryName() const;
+#if WITH_EDITOR
+	virtual FText GetSectionText() const override;
+#endif
+	//~ End UDevelopperSettings interface
+
+	/** Returns current SessionId either based on settings or overriden by commandline */
+	int32 GetSessionId() const;
+
+public:
+
+	/** If true, message bus will only listen to Testers with same sessionId */
+	UPROPERTY(config, EditAnywhere, Category = "Settings")
+	bool bUseSessionId = true;
+
+protected:
+	/**
+	 * The projects SessionId to differentiate data sent over network.
+	 * @note It may be overriden via the command line, "-MessageBusTesterSessionId=1"
+	 */
+	UPROPERTY(config, EditAnywhere, Category = "Settings")
+	int32 SessionId = INDEX_NONE;
+
+public:
+
+	/** Interval threshold between KeepAlive reception before considering another tester as dropped */
+	UPROPERTY(config, EditAnywhere, Category = "Settings", meta=(Units="s"))
+	float TimeoutInterval = 10.0f;
+
+	/** Interval between each discovery signal */
+	UPROPERTY(config, EditAnywhere, Category = "Settings", meta = (Units = "s"))
+	float DiscoveryMessageInterval = 2.0f;
+
+	/**
+	 * The current SessionId in a MessageBusTesting context read from the command line.
+	 * ie. "-MessageBusTesterSessionId=1"
+	 */
+	TOptional<int32> CommandLineSessionId;
+
+	/**
+	 * A friendly name for that instance given through command line (-MessageBusTesterFriendlyName=) to identify it when monitoring.
+	 * If none, by default it will be filled with MachineName:ProcessId
+	 */
+	FName CommandLineFriendlyName;
+};
+
diff --git a/Engine/Plugins/Experimental/MessageBusTester/Source/MessageBusTesterEditor/MessageBusTesterEditor.Build.cs b/Engine/Plugins/Experimental/MessageBusTester/Source/MessageBusTesterEditor/MessageBusTesterEditor.Build.cs
new file mode 100644
--- /dev/null
+++ b/Engine/Plugins/Experimental/MessageBusTester/Source/MessageBusTesterEditor/MessageBusTesterEditor.Build.cs
@@ -0,0 +1,37 @@
+// Copyright Epic Games, Inc. All Rights Reserved.
+
+using UnrealBuildTool;
+
+public class MessageBusTesterEditor : ModuleRules
+{
+	public MessageBusTesterEditor(ReadOnlyTargetRules Target) : base(Target)
+	{
+		PublicDependencyModuleNames.AddRange(
+			new string[]
+			{
+			}
+		);
+
+		PrivateDependencyModuleNames.AddRange(
+			new string[]
+			{
+				"Core",
+				"CoreUObject",
+				"DesktopPlatform",
+				"EditorStyle",
+ 				"InputCore",
+ 				"MessageLog",
+				"Projects",
+				"RenderCore",
+				"Settings",
+				"Slate",
+				"SlateCore",
+				"MessageBusTester",
+                "WorkspaceMenuStructure",
+				"UdpMessaging",
+				"Messaging",
+				"Networking"
+            }
+		);
+	}
+}
diff --git a/Engine/Plugins/Experimental/MessageBusTester/Source/MessageBusTesterEditor/Private/MessageBusTesterEditorModule.cpp b/Engine/Plugins/Experimental/MessageBusTester/Source/MessageBusTesterEditor/Private/MessageBusTesterEditorModule.cpp
new file mode 100644
--- /dev/null
+++ b/Engine/Plugins/Experimental/MessageBusTester/Source/MessageBusTesterEditor/Private/MessageBusTesterEditorModule.cpp
@@ -0,0 +1,37 @@
+// Copyright Epic Games, Inc. All Rights Reserved.
+
+#include "MessageBusTesterEditorModule.h"
+
+#include "SMessageBusTesterPanel.h"
+#include "UObject/ObjectMacros.h"
+#include "UObject/Object.h"
+#include "WorkspaceMenuStructure.h"
+#include "WorkspaceMenuStructureModule.h"
+#include "Framework/Docking/TabManager.h"
+
+DEFINE_LOG_CATEGORY(LogMessageBusTesterEditor);
+
+#define LOCTEXT_NAMESPACE "MessageBusTesterEditor"
+
+void FMessageBusTesterEditorModule::DisplayMessageBusTester()
+{
+	FGlobalTabmanager::Get()->TryInvokeTab(SMessageBusTesterPanel::GetTabName());
+}
+
+void FMessageBusTesterEditorModule::StartupModule()
+{
+	SMessageBusTesterPanel::RegisterNomadTabSpawner(WorkspaceMenu::GetMenuStructure().GetDeveloperToolsMiscCategory());
+}
+
+void FMessageBusTesterEditorModule::ShutdownModule()
+{
+	if (!IsRunningCommandlet() && UObjectInitialized() && !IsEngineExitRequested())
+	{
+		SMessageBusTesterPanel::UnregisterNomadTabSpawner();
+	}
+}
+
+
+#undef LOCTEXT_NAMESPACE
+
+IMPLEMENT_MODULE(FMessageBusTesterEditorModule, MessageBusTesterEditor)
diff --git a/Engine/Plugins/Experimental/MessageBusTester/Source/MessageBusTesterEditor/Private/MessageBusTesterEditorModule.h b/Engine/Plugins/Experimental/MessageBusTester/Source/MessageBusTesterEditor/Private/MessageBusTesterEditorModule.h
new file mode 100644
--- /dev/null
+++ b/Engine/Plugins/Experimental/MessageBusTester/Source/MessageBusTesterEditor/Private/MessageBusTesterEditorModule.h
@@ -0,0 +1,54 @@
+// Copyright Epic Games, Inc. All Rights Reserved.
+
+#pragma once
+
+#include "CoreMinimal.h"
+
+#include "Modules/ModuleManager.h"
+#include "Modules/ModuleInterface.h"
+#include "Features/IModularFeatures.h"
+
+#include "IMessageBusTesterEditorModule.h"
+#include "INetworkMessagingExtension.h"
+
+#include "Misc/ScopeExit.h"
+
+DECLARE_LOG_CATEGORY_EXTERN(LogMessageBusTesterEditor, Log, All);
+
+class FMessageBusTesterEditorModule : public IMessageBusTesterEditorModule
+{
+public:
+	virtual void DisplayMessageBusTester() override;
+
+	static INetworkMessagingExtension* GetMessagingStatistics()
+	{
+		IModularFeatures& ModularFeatures = IModularFeatures::Get();
+		if (IsInGameThread())
+		{
+			if (ModularFeatures.IsModularFeatureAvailable(INetworkMessagingExtension::ModularFeatureName))
+			{
+				return &ModularFeatures.GetModularFeature<INetworkMessagingExtension>(INetworkMessagingExtension::ModularFeatureName);
+			}
+		}
+		else
+		{
+			ModularFeatures.LockModularFeatureList();
+			ON_SCOPE_EXIT
+			{
+				ModularFeatures.UnlockModularFeatureList();
+			};
+		
+			if (ModularFeatures.IsModularFeatureAvailable(INetworkMessagingExtension::ModularFeatureName))
+			{
+				return &ModularFeatures.GetModularFeature<INetworkMessagingExtension>(INetworkMessagingExtension::ModularFeatureName);
+			}
+		}
+		ensureMsgf(false, TEXT("Feature %s is unavailable"), *INetworkMessagingExtension::ModularFeatureName.ToString());
+		return nullptr;
+	}
+
+	//~ Begin IModuleInterface
+	virtual void StartupModule() override;
+	virtual void ShutdownModule() override;
+	//~End IModuleInterface
+};
diff --git a/Engine/Plugins/Experimental/MessageBusTester/Source/MessageBusTesterEditor/Private/MessageBusTesterEditorStyle.cpp b/Engine/Plugins/Experimental/MessageBusTester/Source/MessageBusTesterEditor/Private/MessageBusTesterEditorStyle.cpp
new file mode 100644
--- /dev/null
+++ b/Engine/Plugins/Experimental/MessageBusTester/Source/MessageBusTesterEditor/Private/MessageBusTesterEditorStyle.cpp
@@ -0,0 +1,63 @@
+// Copyright Epic Games, Inc. All Rights Reserved.
+
+#include "MessageBusTesterEditorStyle.h"
+
+#include "Interfaces/IPluginManager.h"
+#include "Styling/SlateTypes.h"
+#include "Styling/CoreStyle.h"
+#include "Styling/SlateStyleRegistry.h"
+
+
+FMessageBusTesterEditorStyle::FMessageBusTesterEditorStyle()
+	: FSlateStyleSet("MessageBusTesterEditorStyle")
+{
+	if (IPluginManager::Get().FindPlugin(TEXT("MessageBusTester")))
+	{
+		const FString ContentDir = IPluginManager::Get().FindPlugin(TEXT("MessageBusTester"))->GetContentDir();
+		SetContentRoot(ContentDir);
+		SetCoreContentRoot(FPaths::EngineContentDir() / TEXT("Slate"));
+	}
+
+	// ListView
+	{
+		const FVector2D Icon8x8(8.0f, 8.0f);
+		FTableRowStyle CriticalStateRow(FCoreStyle::Get().GetWidgetStyle<FTableRowStyle>("TableView.Row"));
+		CriticalStateRow.SetEvenRowBackgroundBrush(FSlateColorBrush(FLinearColor(0.2f, 0.0f, 0.0f, 0.1f)))
+			.SetOddRowBackgroundBrush(FSlateColorBrush(FLinearColor(0.2f, 0.0f, 0.0f, 0.1f)));
+
+		Set("TableView.CriticalStateRow", CriticalStateRow);
+	}
+
+	//Checkbox for live/preview mode
+	{
+		FCheckBoxStyle SwitchStyle = FCheckBoxStyle()
+		.SetForegroundColor(FLinearColor::White)
+		.SetUncheckedImage(FSlateImageBrush(RootToContentDir(TEXT("Switch_OFF.png")), FVector2D(28.0F, 14.0F)))
+		.SetUncheckedHoveredImage(FSlateImageBrush(RootToContentDir(TEXT("Switch_OFF.png")), FVector2D(28.0F, 14.0F)))
+		.SetUncheckedPressedImage(FSlateImageBrush(RootToContentDir(TEXT("Switch_OFF.png")), FVector2D(28.0F, 14.0F)))
+		.SetCheckedImage(FSlateImageBrush(RootToContentDir(TEXT("Switch_ON.png")), FVector2D(28.0F, 14.0F)))
+		.SetCheckedHoveredImage(FSlateImageBrush(RootToContentDir(TEXT("Switch_ON.png")), FVector2D(28.0F, 14.0F)))
+		.SetCheckedPressedImage(FSlateImageBrush(RootToContentDir(TEXT("Switch_ON.png")), FVector2D(28.0F, 14.0F)))
+		.SetPadding(FMargin(0, 0, 0, 1));
+
+		Set("ViewMode", SwitchStyle);
+	}
+
+	
+	FSlateStyleRegistry::RegisterSlateStyle(*this);
+
+#undef IMAGE_BRUSH
+}
+
+FMessageBusTesterEditorStyle::~FMessageBusTesterEditorStyle()
+{
+	FSlateStyleRegistry::UnRegisterSlateStyle(*this);
+}
+
+FMessageBusTesterEditorStyle& FMessageBusTesterEditorStyle::Get()
+{
+	static FMessageBusTesterEditorStyle Inst;
+	return Inst;
+}
+
+
diff --git a/Engine/Plugins/Experimental/MessageBusTester/Source/MessageBusTesterEditor/Private/MessageBusTesterEditorStyle.h b/Engine/Plugins/Experimental/MessageBusTester/Source/MessageBusTesterEditor/Private/MessageBusTesterEditorStyle.h
new file mode 100644
--- /dev/null
+++ b/Engine/Plugins/Experimental/MessageBusTester/Source/MessageBusTesterEditor/Private/MessageBusTesterEditorStyle.h
@@ -0,0 +1,20 @@
+// Copyright Epic Games, Inc. All Rights Reserved.
+
+#pragma once
+
+#include "Styling/SlateStyle.h"
+
+/**
+ * Slate style set that defines all the styles for the MessageBus Tester UI
+ */
+class FMessageBusTesterEditorStyle : public FSlateStyleSet
+{
+public:
+	/** Access the singleton instance for this style set */
+	static FMessageBusTesterEditorStyle& Get();
+
+
+private:
+	FMessageBusTesterEditorStyle();
+	~FMessageBusTesterEditorStyle();
+};
diff --git a/Engine/Plugins/Experimental/MessageBusTester/Source/MessageBusTesterEditor/Private/SMessageBusTestLogger.cpp b/Engine/Plugins/Experimental/MessageBusTester/Source/MessageBusTesterEditor/Private/SMessageBusTestLogger.cpp
new file mode 100644
--- /dev/null
+++ b/Engine/Plugins/Experimental/MessageBusTester/Source/MessageBusTesterEditor/Private/SMessageBusTestLogger.cpp
@@ -0,0 +1,197 @@
+// Copyright Epic Games, Inc. All Rights Reserved.
+
+#include "SMessageBusTestLogger.h"
+
+#include "Modules/ModuleManager.h"
+#include "Widgets/Layout/SBorder.h"
+#include "IMessageBusTesterModule.h"
+#include "IMessageBusTesterLogger.h"
+#include "Widgets/Input/SButton.h"
+#include "Widgets/Text/STextBlock.h"
+#include "Widgets/Views/SListView.h"
+
+#define LOCTEXT_NAMESPACE "SMessageBusTesterLogger"
+
+
+namespace MessageBusTesterLoggerUtils
+{
+	const FName HeaderIdName_SourceTester = "Source";
+	const FName HeaderIdName_LogMessage = "LogMessage";
+}
+
+
+
+/**
+ * TableRow
+ */
+void SMessageBusTesterLogEntryTableRow::Construct(const FArguments& InArgs, const TSharedRef<STableViewBase>& InOwerTableView)
+{
+	LogEntry = InArgs._LogEntry;
+
+	Super::FArguments Arg;
+
+	Super::Construct(Arg, InOwerTableView);
+}
+
+TSharedRef<SWidget> SMessageBusTesterLogEntryTableRow::GenerateWidgetForColumn(const FName& ColumnName)
+{
+	if (MessageBusTesterLoggerUtils::HeaderIdName_SourceTester == ColumnName)
+	{
+		return SNew(SHorizontalBox)
+			+ SHorizontalBox::Slot()
+			.HAlign(HAlign_Left)
+			.VAlign(VAlign_Center)
+			.AutoWidth()
+			[
+				SNew(STextBlock)
+				.Text(this, &SMessageBusTesterLogEntryTableRow::GetSourceTester)
+			];
+	}
+	if (MessageBusTesterLoggerUtils::HeaderIdName_LogMessage == ColumnName)
+	{
+		return SNew(SHorizontalBox)
+			+ SHorizontalBox::Slot()
+			.HAlign(HAlign_Left)
+			.VAlign(VAlign_Center)
+			.AutoWidth()
+			[
+				SNew(STextBlock)
+				.Text(this, &SMessageBusTesterLogEntryTableRow::GetLogMessage)
+				.ColorAndOpacity(this, &SMessageBusTesterLogEntryTableRow::GetLogMessageTextColor)
+			];
+	}
+
+	return SNullWidget::NullWidget;
+}
+
+FText SMessageBusTesterLogEntryTableRow::GetSourceTester() const
+{
+	if (TSharedPtr<FMessageBusTesterLogEntry> EntryPtr = LogEntry.Pin())
+	{
+		return FText::FromName(EntryPtr->Source);
+	}
+
+	return FText::GetEmpty();
+}
+
+FText SMessageBusTesterLogEntryTableRow::GetLogMessage() const
+{
+	if (TSharedPtr<FMessageBusTesterLogEntry> EntryPtr = LogEntry.Pin())
+	{
+		return FText::FromString(EntryPtr->LogMessage);
+	}
+
+	return FText::GetEmpty();
+}
+
+
+FSlateColor SMessageBusTesterLogEntryTableRow::GetLogMessageTextColor() const
+{
+	if (TSharedPtr<FMessageBusTesterLogEntry> EntryPtr = LogEntry.Pin())
+	{
+		switch (EntryPtr->MessageSeverity) 
+		{
+			case EMessageSeverity::Warning: return FLinearColor::Yellow;
+			case EMessageSeverity::Error: return FLinearColor::Red;
+			case EMessageSeverity::PerformanceWarning: return FLinearColor::Yellow;
+			case EMessageSeverity::Info: return FSlateColor::UseForeground();
+			default: return FSlateColor::UseForeground();
+		}
+	}
+
+	return FLinearColor::Gray;
+}
+
+SMessageBusTestLogger::~SMessageBusTestLogger()
+{
+	if (IMessageBusTesterModule::IsAvailable())
+	{
+		IMessageBusTesterModule::Get().GetLogger().OnMessageBusTesterNewLogReceived().RemoveAll(this);
+		IMessageBusTesterModule::Get().GetLogger().OnMessageBusTesterLogCleared().RemoveAll(this);
+	}
+}
+
+void SMessageBusTestLogger::Construct(const FArguments& InArgs)
+{
+	ChildSlot
+	[
+		SNew(SVerticalBox)
+		+ SVerticalBox::Slot()
+		.VAlign(VAlign_Fill)
+		.Padding(10.f, 0.f, 10.f, 10.f)
+		[
+			SNew(SBorder)
+			//.BorderImage(FEditorStyle::GetBrush("Menu.Background"))
+			.Padding(2.f)
+			[
+				SAssignNew(LogEntriesListView, SListView<TSharedPtr<FMessageBusTesterLogEntry>>)
+				.ListItemsSource(&LogEntries)
+				.OnGenerateRow(this, &SMessageBusTestLogger::OnGenerateActivityRowWidget)
+				.SelectionMode(ESelectionMode::None)
+				//.AllowOverscroll(EAllowOverscroll::No)
+				.HeaderRow
+				(
+					SNew(SHeaderRow)
+					+ SHeaderRow::Column(MessageBusTesterLoggerUtils::HeaderIdName_SourceTester)
+					.FillWidth(25.f)
+					.DefaultLabel(LOCTEXT("HeaderName_SourceName", "Tester Name"))
+					+ SHeaderRow::Column(MessageBusTesterLoggerUtils::HeaderIdName_LogMessage)
+					.FillWidth(75.f)
+					.DefaultLabel(LOCTEXT("HeaderName_LogMessage", "Message"))
+				)
+			]
+		]
+		+ SVerticalBox::Slot()
+		.AutoHeight()
+		[
+			SNew(SHorizontalBox)
+			+ SHorizontalBox::Slot()
+			.HAlign(HAlign_Right)
+			[
+				SNew(SButton)
+				.ContentPadding(FMargin(4.f, 2.f))
+				.OnClicked(this, &SMessageBusTestLogger::OnLogClearedClicked)
+				[
+					SNew(STextBlock)
+					.Text(LOCTEXT("ClearLabel", "Clear"))
+				]
+			]
+		]
+	];
+
+	IMessageBusTesterModule::Get().GetLogger().OnMessageBusTesterNewLogReceived().AddSP(this, &SMessageBusTestLogger::OnNewLogReceived);
+	IMessageBusTesterModule::Get().GetLogger().OnMessageBusTesterLogCleared().AddSP(this, &SMessageBusTestLogger::OnLogCleared);
+}
+
+void SMessageBusTestLogger::Tick(const FGeometry& AllottedGeometry, const double InCurrentTime, const float InDeltaTime)
+{
+
+}
+
+TSharedRef<ITableRow> SMessageBusTestLogger::OnGenerateActivityRowWidget(TSharedPtr<FMessageBusTesterLogEntry> InLogEntry, const TSharedRef<STableViewBase>& OwnerTable)
+{
+	TSharedRef<SMessageBusTesterLogEntryTableRow> Row = SNew(SMessageBusTesterLogEntryTableRow, OwnerTable)
+		.LogEntry(InLogEntry);
+	return Row;
+}
+
+void SMessageBusTestLogger::OnNewLogReceived(TSharedRef<FMessageBusTesterLogEntry> NewLog)
+{
+	LogEntries.Insert(NewLog, 0);
+	LogEntriesListView->RequestListRefresh();
+}
+
+void SMessageBusTestLogger::OnLogCleared()
+{
+	LogEntries.Empty();
+	LogEntriesListView->RequestListRefresh();
+}
+
+FReply SMessageBusTestLogger::OnLogClearedClicked()
+{
+	IMessageBusTesterModule::Get().GetLogger().ClearLog();
+	return FReply::Handled();
+}
+
+
+#undef LOCTEXT_NAMESPACE
diff --git a/Engine/Plugins/Experimental/MessageBusTester/Source/MessageBusTesterEditor/Private/SMessageBusTestLogger.h b/Engine/Plugins/Experimental/MessageBusTester/Source/MessageBusTesterEditor/Private/SMessageBusTestLogger.h
new file mode 100644
--- /dev/null
+++ b/Engine/Plugins/Experimental/MessageBusTester/Source/MessageBusTesterEditor/Private/SMessageBusTestLogger.h
@@ -0,0 +1,68 @@
+// Copyright Epic Games, Inc. All Rights Reserved.
+
+#pragma once
+
+
+#include "CoreMinimal.h"
+#include "IMessageBusTesterLogger.h"
+#include "Templates/SharedPointer.h"
+#include "Widgets/DeclarativeSyntaxSupport.h"
+#include "Widgets/SCompoundWidget.h"
+#include "Widgets/Views/STableRow.h"
+#include "Widgets/Views/STableViewBase.h"
+#include "Widgets/Views/SListView.h"
+
+struct FMessageBusTesterLogEntry;
+
+using FMessageBusTesterLogEntryPtr = TSharedPtr<FMessageBusTesterLogEntry>;
+
+class SMessageBusTesterLogEntryTableRow : public SMultiColumnTableRow<FMessageBusTesterLogEntryPtr>
+{
+	using Super = SMultiColumnTableRow<FMessageBusTesterLogEntryPtr>;
+
+public:
+	SLATE_BEGIN_ARGS(SMessageBusTesterLogEntryTableRow) { }
+		SLATE_ARGUMENT(TSharedPtr<FMessageBusTesterLogEntry>, LogEntry)
+	SLATE_END_ARGS()
+
+	void Construct(const FArguments& Args, const TSharedRef<STableViewBase>& OwnerTable);
+
+private:
+	virtual TSharedRef<SWidget> GenerateWidgetForColumn(const FName& ColumnName) override;
+
+	FText GetSourceTester() const;
+	FText GetLogMessage() const;
+	FSlateColor GetLogMessageTextColor() const;
+
+private:
+	TWeakPtr<FMessageBusTesterLogEntry> LogEntry;
+};
+
+
+class SMessageBusTestLogger : public SCompoundWidget
+{
+public:
+	SLATE_BEGIN_ARGS(SMessageBusTestLogger)
+	{}
+	SLATE_END_ARGS()
+
+	virtual ~SMessageBusTestLogger();
+
+	/** Constructs this widget with InArgs */
+	void Construct(const FArguments& InArgs);
+
+	//~ Begin SCompoundWidget interface
+	virtual void Tick(const FGeometry& AllottedGeometry, const double InCurrentTime, const float InDeltaTime) override;
+	//~End SCompoundWidget interface
+
+private:
+	TSharedRef<ITableRow> OnGenerateActivityRowWidget(TSharedPtr<FMessageBusTesterLogEntry> InLogEntry, const TSharedRef<STableViewBase>& OwnerTable);
+	void OnNewLogReceived(TSharedRef<FMessageBusTesterLogEntry> NewLog);
+	void OnLogCleared();
+	FReply OnLogClearedClicked();
+
+private:
+
+	TSharedPtr<SListView<TSharedPtr<FMessageBusTesterLogEntry>>> LogEntriesListView;
+	TArray<TSharedPtr<FMessageBusTesterLogEntry>> LogEntries;
+};
diff --git a/Engine/Plugins/Experimental/MessageBusTester/Source/MessageBusTesterEditor/Private/SMessageBusTestNetwork.cpp b/Engine/Plugins/Experimental/MessageBusTester/Source/MessageBusTesterEditor/Private/SMessageBusTestNetwork.cpp
new file mode 100644
--- /dev/null
+++ b/Engine/Plugins/Experimental/MessageBusTester/Source/MessageBusTesterEditor/Private/SMessageBusTestNetwork.cpp
@@ -0,0 +1,391 @@
+// Copyright Epic Games, Inc. All Rights Reserved.
+
+#pragma once
+
+#include "SMessageBusTestNetwork.h"
+
+#include "MessageBusTesterEditorModule.h"
+
+#include "HAL/Platform.h"
+
+#include "Algo/BinarySearch.h"
+
+#include "Framework/Application/SlateApplication.h"
+#include "Framework/Docking/TabManager.h"
+#include "Framework/MultiBox/MultiBoxBuilder.h"
+
+#include "INetworkMessagingExtension.h"
+
+#include "SlateOptMacros.h"
+
+#include "Types/SlateEnums.h"
+#include "UObject/NameTypes.h"
+#include "Widgets/SBoxPanel.h"
+
+#define LOCTEXT_NAMESPACE "SMessageBusTestNetwork"
+
+namespace MessageBusTesterNetworkUtils
+{
+	const FName HeaderIdName_MessageId    = TEXT("MessageId");
+	const FName HeaderIdName_SentSegments = TEXT("Sent");
+	const FName HeaderIdName_AckSegments  = TEXT("SegmentsAck");
+	const FName HeaderIdName_TotalSize    = TEXT("Size");
+	const FName HeaderIdName_DataRate     = TEXT("DataRate");
+}
+
+void SMessageBusTesterSegmentersTableRow::Construct(const FArguments& InArgs, const TSharedRef<STableViewBase>& InOwerTableView)
+{
+	Segmenter = InArgs._Segmenter;
+
+	Super::FArguments Arg;
+
+	Super::Construct(Arg, InOwerTableView);
+}
+
+TSharedRef<SWidget> SMessageBusTesterSegmentersTableRow::GenerateWidgetForColumn(const FName& ColumnName)
+{
+	using namespace MessageBusTesterNetworkUtils;
+	if (HeaderIdName_MessageId == ColumnName)
+	{
+		return SNew(SHorizontalBox)
+			+ SHorizontalBox::Slot()
+			.HAlign(HAlign_Left)
+			.VAlign(VAlign_Center)
+			.AutoWidth()
+			[
+				SNew(STextBlock)
+				.Text(this, &SMessageBusTesterSegmentersTableRow::GetMessageId)
+			];
+	}
+	if (HeaderIdName_SentSegments == ColumnName)
+	{
+		return SNew(SHorizontalBox)
+			+ SHorizontalBox::Slot()
+			.HAlign(HAlign_Left)
+			.VAlign(VAlign_Center)
+			.AutoWidth()
+			[
+				SNew(STextBlock)
+				.Text(this, &SMessageBusTesterSegmentersTableRow::GetSentSegments)
+			];
+	}
+	if (HeaderIdName_AckSegments == ColumnName)
+	{
+		return SNew(SHorizontalBox)
+			+ SHorizontalBox::Slot()
+			.HAlign(HAlign_Left)
+			.VAlign(VAlign_Center)
+			.AutoWidth()
+			[
+				SNew(STextBlock)
+				.Text(this, &SMessageBusTesterSegmentersTableRow::GetAckSegments)
+			];
+	}
+	if (HeaderIdName_TotalSize == ColumnName)
+	{
+		return SNew(SHorizontalBox)
+			+ SHorizontalBox::Slot()
+			.HAlign(HAlign_Left)
+			.VAlign(VAlign_Center)
+			.AutoWidth()
+			[
+				SNew(STextBlock)
+				.Text(this, &SMessageBusTesterSegmentersTableRow::GetSize)
+			];
+	}
+	if (HeaderIdName_DataRate == ColumnName)
+	{
+		return SNew(SHorizontalBox)
+			+ SHorizontalBox::Slot()
+			.HAlign(HAlign_Left)
+			.VAlign(VAlign_Center)
+			.AutoWidth()
+			[
+				SNew(STextBlock)
+				.Text(this, &SMessageBusTesterSegmentersTableRow::GetDataRate)
+			];
+	}
+	return SNullWidget::NullWidget;
+}
+
+FText SMessageBusTesterSegmentersTableRow::GetMessageId() const
+{
+	if(TSharedPtr<FSegmenterRowData> SegStats = Segmenter.Pin() )
+	{
+		return FText::AsNumber(SegStats->Stats.MessageId);
+	}
+	return FText::GetEmpty();
+}
+
+FText SMessageBusTesterSegmentersTableRow::GetSentSegments() const
+{
+	if(TSharedPtr<FSegmenterRowData> SegStats = Segmenter.Pin() )
+	{
+		return FText::AsNumber(SegStats->Stats.BytesSent);
+	}
+	return FText::GetEmpty();
+}
+
+FText SMessageBusTesterSegmentersTableRow::GetAckSegments() const
+{
+	if(TSharedPtr<FSegmenterRowData> SegStats = Segmenter.Pin() )
+	{
+		return FText::AsNumber(SegStats->Stats.BytesAcknowledged);
+	}
+	return FText::GetEmpty();
+}
+
+FText SMessageBusTesterSegmentersTableRow::GetSize() const
+{
+	if(TSharedPtr<FSegmenterRowData> SegStats = Segmenter.Pin() )
+	{
+		return FText::AsNumber(SegStats->Stats.BytesToSend);
+	}
+	return FText::GetEmpty();
+}
+
+FText SMessageBusTesterSegmentersTableRow::GetDataRate() const
+{
+	if(TSharedPtr<FSegmenterRowData> SegStats = Segmenter.Pin() )
+	{
+		FNumberFormattingOptions Options;
+		Options.MaximumFractionalDigits = 2;
+		Options.MinimumFractionalDigits = 2;
+		return FText::Format(LOCTEXT("BitsPerSec", "{0}"), FText::AsNumber(SegStats->ComputeTransferRate(), &Options) );
+	}
+	return LOCTEXT("ZeroBits","0");
+}
+
+BEGIN_SLATE_FUNCTION_BUILD_OPTIMIZATION
+
+void SMessageBusTestNetwork::Construct(const FArguments& InArgs)
+{
+	using namespace MessageBusTesterNetworkUtils;
+
+	ChildSlot
+	[
+		SNew(SVerticalBox)
+		+ SVerticalBox::Slot()
+		[
+			SNew(SVerticalBox)
+			// Toolbar
+			+ SVerticalBox::Slot()
+			.HAlign(HAlign_Left)
+			.AutoHeight()
+			[
+				SNew(STextBlock)
+				.Text(this, &SMessageBusTestNetwork::GetIPAddress)
+			]
+			+ SVerticalBox::Slot()
+			.HAlign(HAlign_Left)
+			.AutoHeight()
+			[
+				SNew(STextBlock)
+				.Text(this, &SMessageBusTestNetwork::GetWindowSize)
+			]
+			+ SVerticalBox::Slot()
+			.HAlign(HAlign_Left)
+			.AutoHeight()
+			[
+				SNew(STextBlock)
+				.Text(this, &SMessageBusTestNetwork::GetAverageRTT)
+			]
+			+ SVerticalBox::Slot()
+			.HAlign(HAlign_Left)
+			.AutoHeight()
+			[
+			SNew(STextBlock)
+				.Text(this, &SMessageBusTestNetwork::GetSegmentsInFlight)
+			]
+			+ SVerticalBox::Slot()
+			.HAlign(HAlign_Left)
+			.AutoHeight()
+			[
+				SNew(STextBlock)
+				.Text(this, &SMessageBusTestNetwork::GetBytesSent)
+			]
+			+ SVerticalBox::Slot()
+			.HAlign(HAlign_Left)
+			.AutoHeight()
+			[
+				SNew(STextBlock)
+				.Text(this, &SMessageBusTestNetwork::GetSegmentsLost)
+			]
+		]
+		+ SVerticalBox::Slot()
+		[
+			SNew(SBorder)
+			.Padding(2.f)
+			[
+				SAssignNew(SegmenterListView, SListView<TSharedPtr<FSegmenterRowData>>)
+				.ListItemsSource(GetSegmenterListSource())
+				.OnGenerateRow(this, &SMessageBusTestNetwork::OnGenerateActivityRowWidget)
+				.SelectionMode(ESelectionMode::None)
+				//.AllowOverscroll(EAllowOverscroll::No)
+				.HeaderRow
+				(
+					SNew(SHeaderRow)
+					+ SHeaderRow::Column( HeaderIdName_MessageId    )
+					.FillWidth(2.f)
+					.DefaultLabel(LOCTEXT("HeaderName_MessageId", "Id"))
+					+ SHeaderRow::Column( HeaderIdName_SentSegments )
+					.FillWidth(2.f)
+					.DefaultLabel(LOCTEXT("HeaderName_SentSegments", "Sent"))
+					+ SHeaderRow::Column( HeaderIdName_AckSegments  )
+					.FillWidth(2.f)
+					.DefaultLabel(LOCTEXT("HeaderName_AckSegments", "Ack"))
+					+ SHeaderRow::Column( HeaderIdName_TotalSize    )
+					.FillWidth(2.f)
+					.DefaultLabel(LOCTEXT("HeaderName_TotalSize", "Size"))
+					+ SHeaderRow::Column( HeaderIdName_DataRate     )
+					.FillWidth(2.f)
+					.DefaultLabel(LOCTEXT("HeaderName_Rate", "Rate (Mbit/s)"))
+				)
+			]
+		]
+	];
+
+	if (INetworkMessagingExtension* Statistics = FMessageBusTesterEditorModule::GetMessagingStatistics())
+	{
+		Statistics->OnOutboundTransferUpdatedFromThread().AddRaw(this, &SMessageBusTestNetwork::HandleSegmenterUpdated);
+	}
+}
+
+END_SLATE_FUNCTION_BUILD_OPTIMIZATION
+
+
+SMessageBusTestNetwork::~SMessageBusTestNetwork()
+{
+	if (INetworkMessagingExtension* Statistics = FMessageBusTesterEditorModule::GetMessagingStatistics())
+	{
+		Statistics->OnOutboundTransferUpdatedFromThread().RemoveAll(this);
+	}
+}
+
+void SMessageBusTestNetwork::HandleSegmenterUpdated(FOutboundTransferStatistics InStats)
+{
+	OutboundStats.Enqueue(MoveTemp(InStats));
+}
+
+const SMessageBusTestNetwork::FSegmenterStatsArray* SMessageBusTestNetwork::GetSegmenterListSource() const
+{
+	if (Item.IsValid())
+	{
+		TSharedPtr<FDiscoveredTesterTableRowData> Data = Item.Pin();
+		if (Data.IsValid())
+		{
+			if (SMessageBusTestNetwork::FSegmenterStatsArray const * SegmenterData = AllSegmenters.Find(Data->EndpointIdentifier))
+			{
+				return SegmenterData;
+			}	
+		}
+	}
+
+	static FSegmenterStatsArray EmptyArray;
+	return &EmptyArray;
+}
+
+TSharedRef<ITableRow> SMessageBusTestNetwork::OnGenerateActivityRowWidget(TSharedPtr<FSegmenterRowData> InSegmenter, const TSharedRef<STableViewBase>& OwnerTable)
+{
+	TSharedRef<SMessageBusTesterSegmentersTableRow> Row = SNew(SMessageBusTesterSegmentersTableRow, OwnerTable)
+		.Segmenter(InSegmenter);
+	return Row;
+}
+
+void SMessageBusTestNetwork::Tick(const FGeometry& AllottedGeometry, const double InCurrentTime, const float InDeltaTime)
+{
+	TArray<FOutboundTransferStatistics> StatsList;
+	FOutboundTransferStatistics Stats;
+	while (OutboundStats.Dequeue(Stats))
+	{
+		StatsList.Add(Stats);
+	}
+	
+	if (StatsList.Num())
+	{
+		for (FOutboundTransferStatistics SingleStat : StatsList)
+		{
+			FSegmenterStatsArray& Array = AllSegmenters.FindOrAdd(SingleStat.DestinationId);
+			TSharedPtr<FSegmenterRowData> SharedValue = MakeShared<FSegmenterRowData>(SingleStat);
+			// We use upper bound because we want the array to be in reverse order.
+			//
+			auto Pos = Algo::UpperBound(Array, SharedValue, [](const TSharedPtr<FSegmenterRowData>& Value, const TSharedPtr<FSegmenterRowData>& Check)
+			{
+				if (Check.IsValid())
+				{
+					return *Value < *Check;
+				}
+				return false;
+			});
+			if (Array.Num()>0 && Pos < Array.Num() && Array[Pos] && *Array[Pos] == SingleStat)
+			{
+				*Array[Pos] = SingleStat;
+			}
+			else
+			{
+				Array.Insert(MoveTemp(SharedValue),Pos);
+			}
+		}
+		SegmenterListView->RequestListRefresh();
+	}
+
+	Super::Tick(AllottedGeometry, InCurrentTime, InDeltaTime);
+}
+
+
+FMessageTransportStatistics SMessageBusTestNetwork::GetStats() const
+{
+	if (TSharedPtr<FDiscoveredTesterTableRowData> RowData = Item.Pin())
+	{
+		return RowData->LatestStats();
+	}
+	return {};
+}
+
+void SMessageBusTestNetwork::SetNetworkViewData(TWeakPtr<FDiscoveredTesterTableRowData> View)
+{
+	Item = View;
+	if (SegmenterListView)
+	{
+		SegmenterListView->SetItemsSource(GetSegmenterListSource());
+		SegmenterListView->RequestListRefresh();
+	}
+}
+
+FText SMessageBusTestNetwork::GetIPAddress() const
+{
+	FMessageTransportStatistics Stats = GetStats();
+	return FText::Format(LOCTEXT("CurrentIpAddress", "IPv4 Address {0}"), FText::FromString(Stats.IPv4AsString));
+}
+
+FText SMessageBusTestNetwork::GetAverageRTT() const
+{
+	FMessageTransportStatistics Stats = GetStats();
+	return FText::Format(LOCTEXT("CurrentRTT", "Average RTT: {0} ms"), FText::AsNumber(Stats.AverageRTT.GetFractionMilli()));
+}
+
+FText SMessageBusTestNetwork::GetSegmentsInFlight() const
+{
+	FMessageTransportStatistics Stats = GetStats();
+	return FText::Format(LOCTEXT("BytesInFlight", "Bytes Inflight: {0}"), FText::AsNumber(Stats.BytesInflight));
+}
+
+FText SMessageBusTestNetwork::GetWindowSize() const
+{
+	FMessageTransportStatistics Stats = GetStats();
+	return FText::Format(LOCTEXT("WindowSize", "WindowSize : {0}"), FText::AsNumber(Stats.WindowSize));
+}
+
+FText SMessageBusTestNetwork::GetBytesSent() const
+{
+	FMessageTransportStatistics Stats = GetStats();
+	return FText::Format(LOCTEXT("BytesSent", "Bytes Sent : {0}"), FText::AsNumber(Stats.TotalBytesSent));
+}
+
+FText SMessageBusTestNetwork::GetSegmentsLost() const
+{
+	FMessageTransportStatistics Stats = GetStats();
+	return FText::Format(LOCTEXT("BytesLost", "Bytes Lost : {0}"), FText::AsNumber(Stats.TotalBytesLost));
+}
+
+#undef LOCTEXT_NAMESPACE
diff --git a/Engine/Plugins/Experimental/MessageBusTester/Source/MessageBusTesterEditor/Private/SMessageBusTestNetwork.h b/Engine/Plugins/Experimental/MessageBusTester/Source/MessageBusTesterEditor/Private/SMessageBusTestNetwork.h
new file mode 100644
--- /dev/null
+++ b/Engine/Plugins/Experimental/MessageBusTester/Source/MessageBusTesterEditor/Private/SMessageBusTestNetwork.h
@@ -0,0 +1,135 @@
+// Copyright Epic Games, Inc. All Rights Reserved.
+
+#pragma once
+
+#include "Containers/SpscQueue.h"
+
+#include "INetworkMessagingExtension.h"
+#include "Framework/Docking/TabManager.h"
+
+#include "Misc/DateTime.h"
+#include "Styling/SlateColor.h"
+#include "Styling/SlateTypes.h"
+#include "Templates/SharedPointer.h"
+#include "Widgets/DeclarativeSyntaxSupport.h"
+#include "Widgets/Docking/SDockTab.h"
+#include "Widgets/SDiscoveredTesterListView.h"
+#include "Widgets/Views/STableRow.h"
+#include "Widgets/SCompoundWidget.h"
+
+class FWorkspaceItem;
+
+struct FSegmenterRowData
+{
+	FSegmenterRowData(const FOutboundTransferStatistics& InStats)
+		: Stats(InStats)
+	{
+		UpdateTime = AddTime = FDateTime::UtcNow();
+	}
+
+	bool operator < (const FSegmenterRowData &Other) const
+	{
+		// We want the messages to be in reverse order so we reverse the logic.
+		return Stats.MessageId >= Other.Stats.MessageId;
+	}
+
+	bool operator == (const FOutboundTransferStatistics& Other) const
+	{
+		return Stats.MessageId == Other.MessageId;
+	}
+
+	FSegmenterRowData& operator=(const FOutboundTransferStatistics& InStats)
+	{
+		Stats = InStats;
+		UpdateTime = FDateTime::UtcNow();
+		return *this;
+	}
+
+	double ComputeTransferRate() const
+	{
+		FTimespan Span = UpdateTime - AddTime;
+		double Seconds = Span.GetTotalSeconds();
+		if (Seconds > UE_KINDA_SMALL_NUMBER )
+		{
+			// Compute the number of bits in our segment size.
+			//
+			constexpr double Bits = 8.0;
+			constexpr double MegaBits = 1000000;
+			return (double)(Stats.BytesSent * 8.0) / MegaBits /(double)Seconds;
+		}
+		return 0;
+	}
+	FDateTime		   AddTime;
+	FDateTime		   UpdateTime;
+	FOutboundTransferStatistics Stats;
+};
+
+class SMessageBusTesterSegmentersTableRow : public SMultiColumnTableRow<TSharedRef<FSegmenterRowData>>
+{
+	using Super = SMultiColumnTableRow<TSharedRef<FSegmenterRowData>>;
+
+public:
+	SLATE_BEGIN_ARGS(SMessageBusTesterSegmentersTableRow) { }
+		SLATE_ARGUMENT(TSharedPtr<FSegmenterRowData>, Segmenter)
+	SLATE_END_ARGS()
+
+	void Construct(const FArguments& Args, const TSharedRef<STableViewBase>& OwerTableView);
+
+
+private:
+	virtual TSharedRef<SWidget> GenerateWidgetForColumn(const FName& ColumnName) override;
+
+	FText GetMessageId() const;
+	FText GetSentSegments() const;
+	FText GetAckSegments() const;
+	FText GetSize() const;
+	FText GetDataRate() const;
+
+	TWeakPtr<FSegmenterRowData> Segmenter;
+};
+
+/**
+ * Panel used to show stage monitoring data
+ */
+class SMessageBusTestNetwork : public SCompoundWidget
+{
+public:
+	virtual ~SMessageBusTestNetwork();
+
+private:
+	using Super = SCompoundWidget;
+
+public:
+	SLATE_BEGIN_ARGS(SMessageBusTestNetwork) {}
+	SLATE_END_ARGS()
+
+	void Construct(const FArguments& InArgs);
+
+	FText GetIPAddress() const;
+	FText GetAverageRTT() const;
+	FText GetWindowSize() const;
+	FText GetSegmentsInFlight() const;
+	FText GetBytesSent() const;
+	FText GetSegmentsLost() const;
+
+	FMessageTransportStatistics GetStats() const;
+	void SetNetworkViewData(TWeakPtr<FDiscoveredTesterTableRowData> View);
+
+	using FSegmenterStatsArray = TArray<TSharedPtr<FSegmenterRowData>>;
+	const FSegmenterStatsArray* GetSegmenterListSource() const;
+
+	void Tick(const FGeometry& AllottedGeometry, const double InCurrentTime, const float InDeltaTime);
+
+private:
+	void HandleSegmenterUpdated(FOutboundTransferStatistics InStats);
+
+	TSharedRef<ITableRow> OnGenerateActivityRowWidget(TSharedPtr<FSegmenterRowData> InLogEntry, const TSharedRef<STableViewBase>& OwnerTable);
+
+	TWeakPtr<FDiscoveredTesterTableRowData>   Item;
+
+	TSharedPtr<SListView<TSharedPtr<FSegmenterRowData>>> SegmenterListView;
+
+	TMap<FGuid, FSegmenterStatsArray>	AllSegmenters;
+
+	TSpscQueue<FOutboundTransferStatistics> OutboundStats;
+};
diff --git a/Engine/Plugins/Experimental/MessageBusTester/Source/MessageBusTesterEditor/Private/SMessageBusTestPlan.cpp b/Engine/Plugins/Experimental/MessageBusTester/Source/MessageBusTesterEditor/Private/SMessageBusTestPlan.cpp
new file mode 100644
--- /dev/null
+++ b/Engine/Plugins/Experimental/MessageBusTester/Source/MessageBusTesterEditor/Private/SMessageBusTestPlan.cpp
@@ -0,0 +1,117 @@
+// Copyright Epic Games, Inc. All Rights Reserved.
+
+#include "SMessageBusTestPlan.h"
+
+#include "Framework/Application/SlateApplication.h"
+
+#include "SlateOptMacros.h"
+#include "SMessageBusTestNetwork.h"
+
+#include "Styling/SlateStyle.h"
+#include "Types/SlateEnums.h"
+#include "Widgets/Input/SButton.h"
+#include "Widgets/Layout/SBorder.h"
+#include "Widgets/Layout/SBox.h"
+#include "IMessageBusTesterModule.h"
+#include "IMessageBusTester.h"
+#include "Widgets/SMessageBusTestPlanListView.h"
+#include "Widgets/SAddTestPlanEntryDialog.h"
+#include "DiscoveredTester.h"
+
+#define LOCTEXT_NAMESPACE "SMessageBusTestPlan"
+
+
+
+BEGIN_SLATE_FUNCTION_BUILD_OPTIMIZATION
+
+void SMessageBusTestPlan::Construct(const FArguments& InArgs)
+{
+	ChildSlot
+	[
+		SNew(SVerticalBox)
+		+ SVerticalBox::Slot()
+		.VAlign(VAlign_Fill)
+		[
+			// Test plan entries
+			SNew(SBorder)
+			.HAlign(HAlign_Fill)
+			.VAlign(VAlign_Fill)
+			.Padding(FMargin(4.f, 4.f, 4.f, 4.f))
+			[
+				SNew(SMessageBusTestPlanListView)
+			]
+		]
+		+ SVerticalBox::Slot()
+		.AutoHeight()
+		[
+			SNew(SHorizontalBox)
+			+ SHorizontalBox::Slot()
+			[
+				SNew(SButton)
+				.ContentPadding(FMargin(4.f, 2.f))
+				.OnClicked(this, &SMessageBusTestPlan::OnAddPayloadClicked)
+				.HAlign(HAlign_Center)
+				.VAlign(VAlign_Center)
+				[
+					SNew(STextBlock)
+					.Text(LOCTEXT("AddPayload", "Add"))
+				]
+			]
+			+ SHorizontalBox::Slot()
+			[
+				SNew(SButton)
+				.ContentPadding(FMargin(4.f, 2.f))
+				.OnClicked(this, &SMessageBusTestPlan::OnStartTestClicked)
+				.HAlign(HAlign_Center)
+				.VAlign(VAlign_Center)
+				[
+					SNew(STextBlock)
+					.Text(LOCTEXT("StartActivity", "Start"))
+				]
+			]
+			+ SHorizontalBox::Slot()
+			[
+				SNew(SButton)
+				.ContentPadding(FMargin(4.f, 2.f))
+				.OnClicked(this, &SMessageBusTestPlan::OnStopTestClicked)
+				.HAlign(HAlign_Center)
+				.VAlign(VAlign_Center)
+				[
+					SNew(STextBlock)
+					.Text(LOCTEXT("StopActivity", "Stop"))
+				]
+			]
+		]
+	];
+}
+
+FReply SMessageBusTestPlan::OnStartTestClicked()
+{
+	IMessageBusTesterModule::Get().GetMessageBusTester().StartTest();
+
+
+	return FReply::Handled();
+}
+
+
+FReply SMessageBusTestPlan::OnStopTestClicked()
+{
+	IMessageBusTesterModule::Get().GetMessageBusTester().StopTest();
+	return FReply::Handled();
+}
+
+FReply SMessageBusTestPlan::OnAddPayloadClicked()
+{
+	SAddTestPlanEntryDialog::OpenDialog();
+	return FReply::Handled();
+}
+
+END_SLATE_FUNCTION_BUILD_OPTIMIZATION
+
+
+#undef LOCTEXT_NAMESPACE
+
+
+
+
+
diff --git a/Engine/Plugins/Experimental/MessageBusTester/Source/MessageBusTesterEditor/Private/SMessageBusTestPlan.h b/Engine/Plugins/Experimental/MessageBusTester/Source/MessageBusTesterEditor/Private/SMessageBusTestPlan.h
new file mode 100644
--- /dev/null
+++ b/Engine/Plugins/Experimental/MessageBusTester/Source/MessageBusTesterEditor/Private/SMessageBusTestPlan.h
@@ -0,0 +1,44 @@
+// Copyright Epic Games, Inc. All Rights Reserved.
+
+#pragma once
+
+#include "Widgets/SCompoundWidget.h"
+
+#include "CoreMinimal.h"
+#include "SlateFwd.h"
+#include "Styling/SlateColor.h"
+#include "Styling/SlateTypes.h"
+#include "Widgets/DeclarativeSyntaxSupport.h"
+#include "Templates/SharedPointer.h"
+#include "Widgets/Docking/SDockTab.h"
+#include "Framework/Docking/TabManager.h"
+
+#include "Widgets/Views/STableRow.h"
+#include "Widgets/Views/STreeView.h"
+#include "MessageBusTesterCommon.h"
+
+class FWorkspaceItem;
+
+/**
+ * 
+ */
+class SMessageBusTestPlan : public SCompoundWidget
+{
+public:
+	virtual ~SMessageBusTestPlan() = default;
+
+private:
+	using Super = SCompoundWidget;
+
+public:
+	SLATE_BEGIN_ARGS(SMessageBusTestPlan) {}
+	SLATE_END_ARGS()
+
+	void Construct(const FArguments& InArgs);
+
+private:
+	
+	FReply OnStartTestClicked();
+	FReply OnStopTestClicked();
+	FReply OnAddPayloadClicked();
+};
diff --git a/Engine/Plugins/Experimental/MessageBusTester/Source/MessageBusTesterEditor/Private/SMessageBusTesterPanel.cpp b/Engine/Plugins/Experimental/MessageBusTester/Source/MessageBusTesterEditor/Private/SMessageBusTesterPanel.cpp
new file mode 100644
--- /dev/null
+++ b/Engine/Plugins/Experimental/MessageBusTester/Source/MessageBusTesterEditor/Private/SMessageBusTesterPanel.cpp
@@ -0,0 +1,264 @@
+// Copyright Epic Games, Inc. All Rights Reserved.
+
+#include "SMessageBusTesterPanel.h"
+
+#include "EditorFontGlyphs.h"
+#include "EditorStyleSet.h"
+#include "Styling/AppStyle.h"
+
+#include "Framework/Application/SlateApplication.h"
+#include "Framework/Docking/TabManager.h"
+#include "Framework/MultiBox/MultiBoxBuilder.h"
+#include "IMessageBusTester.h"
+#include "IMessageBusTesterModule.h"
+#include "SlateOptMacros.h"
+#include "Widgets/Docking/SDockTab.h"
+#include "Widgets/Images/SImage.h"
+#include "Widgets/Images/SThrobber.h"
+#include "Widgets/Input/SButton.h"
+#include "Widgets/Layout/SBorder.h"
+#include "Widgets/Layout/SSpacer.h"
+#include "Widgets/SDiscoveredTesterListView.h"
+#include "Widgets/SMessageBusTestPlanListView.h"
+#include "WorkspaceMenuStructure.h"
+#include "WorkspaceMenuStructureModule.h"
+#include "SMessageBusTestPlan.h"
+#include "SMessageBusTestLogger.h"
+
+#define LOCTEXT_NAMESPACE "SMessageBusTesterPanel"
+
+TWeakPtr<SMessageBusTesterPanel> SMessageBusTesterPanel::PanelInstance;
+FDelegateHandle SMessageBusTesterPanel::LevelEditorTabManagerChangedHandle;
+
+
+namespace MessageBusTesterUtilities
+{
+	static const FName NAME_App = FName("MessageBusTesterPanelApp");
+	static const FName NAME_LevelEditorModuleName = FName("LevelEditor");
+
+	TSharedRef<SDockTab> CreateApp(const FSpawnTabArgs& Args)
+	{
+		return SNew(SDockTab)
+			.TabRole(ETabRole::NomadTab)
+			[
+				SNew(SMessageBusTesterPanel)
+			];
+	}
+}
+
+FName SMessageBusTesterPanel::GetTabName()
+{
+	return MessageBusTesterUtilities::NAME_App;
+}
+
+void SMessageBusTesterPanel::RegisterNomadTabSpawner(TSharedRef<FWorkspaceItem> InWorkspaceItem)
+{
+	FTabSpawnerEntry& TabSpawnerEntry = FGlobalTabmanager::Get()->RegisterNomadTabSpawner(MessageBusTesterUtilities::NAME_App, FOnSpawnTab::CreateStatic(&MessageBusTesterUtilities::CreateApp))
+		.SetDisplayName(LOCTEXT("TabTitle", "MessageBus Tester"))
+		.SetTooltipText(LOCTEXT("TooltipText", "Test MessageBus performance"))
+		.SetIcon(FSlateIcon(FAppStyle::Get().GetStyleSetName(), "LevelEditor.GameSettings.Small")); 
+
+#if WITH_EDITOR
+	TabSpawnerEntry.SetGroup(WorkspaceMenu::GetMenuStructure().GetDeveloperToolsMiscCategory());
+#else
+	TabSpawnerEntry.SetGroup(WorkspaceMenu::GetMenuStructure().GetToolsCategory());
+#endif
+}
+
+void SMessageBusTesterPanel::UnregisterNomadTabSpawner()
+{
+	if (FSlateApplication::IsInitialized())
+	{
+		FGlobalTabmanager::Get()->UnregisterNomadTabSpawner(MessageBusTesterUtilities::NAME_App);
+	}
+}
+
+TSharedPtr<SMessageBusTesterPanel> SMessageBusTesterPanel::GetPanelInstance()
+{
+	return SMessageBusTesterPanel::PanelInstance.Pin();
+}
+
+SMessageBusTesterPanel::~SMessageBusTesterPanel()
+{
+	if (DiscoveredTestersList.IsValid())
+	{
+		DiscoveredTestersList->OnSelectionChanged().RemoveAll(this);
+	}
+}
+
+BEGIN_SLATE_FUNCTION_BUILD_OPTIMIZATION
+void SMessageBusTesterPanel::Construct(const FArguments& InArgs)
+{
+	PanelInstance = StaticCastSharedRef<SMessageBusTesterPanel>(AsShared());
+
+	SAssignNew(DiscoveredTestersList, SDiscoveredTesterListView);
+
+	ChildSlot
+	[
+		SNew(SOverlay)
+		+ SOverlay::Slot()
+		[
+			SNew(SVerticalBox)
+			// Toolbar
+			+ SVerticalBox::Slot()
+			.AutoHeight()
+			.Padding(5.f, 5.f, 5.f, 5.f)
+			[
+				MakeToolbarWidget()
+			]
+			+ SVerticalBox::Slot()
+			.Padding(5.f, 5.f, 5.f, 5.f)
+			.FillHeight(.3)
+			[
+				SNew(SSplitter)
+				.Orientation(Orient_Horizontal)
+				+ SSplitter::Slot()
+				.Value(.75f)
+				[
+					// Discovered tester List
+					SNew(SBorder)
+					.HAlign(HAlign_Fill)
+					.VAlign(VAlign_Fill)
+					.Padding(FMargin(4.f, 4.f, 4.f, 4.f))
+					[
+						DiscoveredTestersList.ToSharedRef()
+					]
+				]
+				+ SSplitter::Slot()
+				.Value(.25f)
+				[
+					// Data Provider Activities
+					SNew(SBorder)
+					.HAlign(HAlign_Fill)
+					.VAlign(VAlign_Fill)
+					.Padding(FMargin(4.f, 4.f, 4.f, 4.f))
+					 [
+						 SAssignNew(NetworkView, SMessageBusTestNetwork)
+					 ]
+				]
+			]
+			+ SVerticalBox::Slot()
+			.FillHeight(.3f)
+			[
+				SNew(SHorizontalBox)
+				+ SHorizontalBox::Slot()
+				.FillWidth(.75f)
+				[
+					SAssignNew(Logger, SMessageBusTestLogger)
+				]
+				+ SHorizontalBox::Slot()
+				.FillWidth(.25f)
+				[
+					SAssignNew(TestPlan, SMessageBusTestPlan)
+				]
+			]
+			// This tester status
+			+ SVerticalBox::Slot()
+			.Padding(5.f, 5.f, 5.f, 5.f)
+			.AutoHeight()
+			[
+				SNew(SBorder)
+				.HAlign(HAlign_Fill)
+				.VAlign(VAlign_Fill)
+				.Padding(FMargin(4.f, 4.f, 4.f, 4.f))
+				[
+					SNew(SHorizontalBox)
+					// Spacer
+					+ SHorizontalBox::Slot()
+					[
+						SNew(SSpacer)
+					]
+					// Status
+					+ SHorizontalBox::Slot()
+					.AutoWidth()
+					.HAlign(HAlign_Right)
+					.Padding(10.f, 10.f, 10.f, 10.f)
+					[
+						SNew(STextBlock)
+						.Text(LOCTEXT("TesterStatusLabel", "State : "))
+					]
+					+ SHorizontalBox::Slot()
+					.AutoWidth()
+					.HAlign(HAlign_Right)
+					.Padding(10.f, 10.f, 10.f, 10.f)
+					[
+						SNew(STextBlock)
+						.Text(this, &SMessageBusTesterPanel::GetTesterStatus)
+					]
+				]
+			]
+		]
+	];
+
+	DiscoveredTestersList->OnSelectionChanged().AddRaw(this, &SMessageBusTesterPanel::HandleSelectionChange);
+}
+
+void SMessageBusTesterPanel::HandleSelectionChange(FDiscoveredTesterTableRowDataPtr RowData)
+{
+	if (NetworkView.IsValid())
+	{
+		NetworkView->SetNetworkViewData(RowData);
+	}
+}
+
+TSharedRef<SWidget> SMessageBusTesterPanel::MakeToolbarWidget()
+{
+	TSharedRef<SWidget> Toolbar = SNew(SBorder)
+	.VAlign(VAlign_Center)
+	.Padding(FMargin(5.f, 5.f, 5.f, 5.f))
+	[
+
+		SNew(SHorizontalBox)
+		// Clear entries buttons
+		+ SHorizontalBox::Slot()
+		.AutoWidth()
+		.VAlign(VAlign_Center)
+		.Padding(4.f)
+		[
+			SNew(SHorizontalBox)
+			+ SHorizontalBox::Slot()
+			.FillWidth(1.f)
+			[
+				SNew(SButton)
+				.ButtonStyle(FAppStyle::Get(), "ToggleButton")
+				.ToolTipText(LOCTEXT("ClearButton_ToolTip", "Clear unresponsive providers and all entries."))
+				.ContentPadding(FMargin(4.f, 4.f))
+				.HAlign(HAlign_Center)
+				.VAlign(VAlign_Center)
+				.OnClicked(this, &SMessageBusTesterPanel::OnClearClicked)
+				[
+					SNew(STextBlock)
+					.Font(FAppStyle::Get().GetFontStyle("FontAwesome.14"))
+					.Text(FEditorFontGlyphs::Eraser)
+					.ColorAndOpacity(FLinearColor::White)
+				]
+			]
+		]
+	];
+
+	return Toolbar;
+}
+
+END_SLATE_FUNCTION_BUILD_OPTIMIZATION
+
+FReply SMessageBusTesterPanel::OnClearClicked()
+{
+	return FReply::Handled();
+}
+
+FText SMessageBusTesterPanel::GetTesterStatus() const
+{
+	const EMessageBusTesterState State = IMessageBusTesterModule::Get().GetMessageBusTester().GetState();
+	if (State == EMessageBusTesterState::Active)
+	{
+		return LOCTEXT("MonitorStatusActive", "Active");
+	}
+	else
+	{
+		return LOCTEXT("MonitorStatusInactive", "Inactive");
+	}
+}
+
+
+
+#undef LOCTEXT_NAMESPACE
diff --git a/Engine/Plugins/Experimental/MessageBusTester/Source/MessageBusTesterEditor/Private/SMessageBusTesterPanel.h b/Engine/Plugins/Experimental/MessageBusTester/Source/MessageBusTesterEditor/Private/SMessageBusTesterPanel.h
new file mode 100644
--- /dev/null
+++ b/Engine/Plugins/Experimental/MessageBusTester/Source/MessageBusTesterEditor/Private/SMessageBusTesterPanel.h
@@ -0,0 +1,72 @@
+// Copyright Epic Games, Inc. All Rights Reserved.
+
+#pragma once
+
+#include "Widgets/SCompoundWidget.h"
+
+#include "CoreMinimal.h"
+#include "SlateFwd.h"
+#include "Styling/SlateColor.h"
+#include "Styling/SlateTypes.h"
+#include "Widgets/DeclarativeSyntaxSupport.h"
+#include "Templates/SharedPointer.h"
+#include "Widgets/Docking/SDockTab.h"
+#include "Framework/Docking/TabManager.h"
+#include "SMessageBusTestNetwork.h"
+
+class SMessageBusTestPlan;
+class SDiscoveredTesterListView;
+class FWorkspaceItem;
+class SMessageBusTestLogger;
+
+
+/**
+ * Main panel for Message Bus tester.
+ */
+class SMessageBusTesterPanel : public SCompoundWidget
+{
+public:
+	virtual ~SMessageBusTesterPanel();
+
+	static FName GetTabName();
+	static void RegisterNomadTabSpawner(TSharedRef<FWorkspaceItem> InWorkspaceItem);
+	static void UnregisterNomadTabSpawner();
+	static TSharedPtr<SMessageBusTesterPanel> GetPanelInstance();
+
+private:
+	using Super = SCompoundWidget;
+	static TWeakPtr<SMessageBusTesterPanel> PanelInstance;
+	static FDelegateHandle LevelEditorTabManagerChangedHandle;
+
+public:
+	SLATE_BEGIN_ARGS(SMessageBusTesterPanel) {}
+	SLATE_END_ARGS()
+
+	void Construct(const FArguments& InArgs);
+
+	void HandleSelectionChange(FDiscoveredTesterTableRowDataPtr RowData);
+private:
+
+	/** Make the toolbar widgets */
+	TSharedRef<SWidget> MakeToolbarWidget();
+
+	/** Handles clear everything button clicked */
+	FReply OnClearClicked();
+
+	/** Handles StageMonitor settings button clicked */
+	FReply OnShowProjectSettingsClicked();
+
+	/** Returns tester status */
+	FText GetTesterStatus() const;
+
+	FReply OnStartTestClicked();
+	FReply OnStopTestClicked();
+	FReply OnAddPayloadClicked();
+
+private:
+
+	TSharedPtr<SMessageBusTestNetwork>    NetworkView;
+	TSharedPtr<SDiscoveredTesterListView> DiscoveredTestersList;
+	TSharedPtr<SMessageBusTestPlan> TestPlan;
+	TSharedPtr<SMessageBusTestLogger> Logger;
+};
diff --git a/Engine/Plugins/Experimental/MessageBusTester/Source/MessageBusTesterEditor/Private/Widgets/SAddTestPlanEntryDialog.cpp b/Engine/Plugins/Experimental/MessageBusTester/Source/MessageBusTesterEditor/Private/Widgets/SAddTestPlanEntryDialog.cpp
new file mode 100644
--- /dev/null
+++ b/Engine/Plugins/Experimental/MessageBusTester/Source/MessageBusTesterEditor/Private/Widgets/SAddTestPlanEntryDialog.cpp
@@ -0,0 +1,163 @@
+// Copyright Epic Games, Inc. All Rights Reserved.
+
+#include "SAddTestPlanEntryDialog.h"
+
+#include "Widgets/Input/SButton.h"
+#include "Widgets/Input/SNumericEntryBox.h"
+#include "IMessageBusTester.h"
+#include "IMessageBusTesterModule.h"
+
+#include "Math/UnrealMath.h"
+
+#include "EditorStyleSet.h"
+#include "Styling/AppStyle.h"
+
+#define LOCTEXT_NAMESPACE "AddTestPlanEntryDialog"
+
+TWeakPtr<SWindow> SAddTestPlanEntryDialog::AddPointDialogWindow;
+
+void SAddTestPlanEntryDialog::Construct(const FArguments& InArgs)
+{
+	ChildSlot
+	[
+		SNew(SVerticalBox)
+		+ SVerticalBox::Slot()
+		.Padding(5.0f, 5.0f)
+		.AutoHeight()
+		[
+			SNew(SBorder)
+			.HAlign(HAlign_Fill)
+			.VAlign(VAlign_Fill)
+			.Padding(FMargin(4.f, 4.f, 4.f, 4.f))
+			[
+				SNew(SVerticalBox)
+				+ SVerticalBox::Slot()
+				.Padding(5.0f, 5.0f)
+				.AutoHeight()
+				[
+					SNew(SHorizontalBox)
+					+ SHorizontalBox::Slot()
+					.Padding(FMargin(0, 1, 0, 1))
+					.FillWidth(1)
+					[
+						SNew(STextBlock)
+						.Text(LOCTEXT("PayloadSizeLabel", "Payload size :"))
+					]
+					+ SHorizontalBox::Slot()
+					[
+						SNew(SNumericEntryBox<int32>)
+						.AllowSpin(false)
+						.Value_Lambda([this]() { return TestPlanItem.NumBytes; })
+						.OnValueChanged(SNumericEntryBox<int32>::FOnValueChanged::CreateLambda([this](int32 NewValue) { TestPlanItem.NumBytes = FMath::Max(16,NewValue); }))
+						.OnValueCommitted(SNumericEntryBox<int32>::FOnValueCommitted::CreateLambda([this](int32 NewValue, ETextCommit::Type CommitType) { TestPlanItem.NumBytes = FMath::Max(16,NewValue); }))
+					]
+				]
+				+ SVerticalBox::Slot()
+				.Padding(5.0f, 5.0f)
+				.AutoHeight()
+				[
+					SNew(SHorizontalBox)
+					+ SHorizontalBox::Slot()
+					.Padding(FMargin(0, 1, 0, 1))
+					.FillWidth(1)
+					[
+						SNew(STextBlock)
+						.Text(LOCTEXT("IntervalLabel", "Interval (s) :"))
+					]
+					+ SHorizontalBox::Slot()
+					[
+						SNew(SNumericEntryBox<float>)
+						.AllowSpin(false)
+						.Value_Lambda([this]() { return TestPlanItem.IntervalSeconds; })
+						.OnValueChanged(SNumericEntryBox<float>::FOnValueChanged::CreateLambda([this](float NewValue) { TestPlanItem.IntervalSeconds = FMath::Max(0.0,NewValue); }))
+						.OnValueCommitted(SNumericEntryBox<float>::FOnValueCommitted::CreateLambda([this](float NewValue, ETextCommit::Type CommitType) { TestPlanItem.IntervalSeconds = FMath::Max(0.0,NewValue); }))
+					]
+				]
+			]
+		]
+		+ SVerticalBox::Slot()
+		.Padding(5.0f, 5.0f)
+		.AutoHeight()
+		[
+			SNew(SBorder)
+			.HAlign(HAlign_Fill)
+			.VAlign(VAlign_Fill)
+			.Padding(FMargin(4.f, 4.f, 4.f, 4.f))
+			[
+				MakeButtonsWidget()
+			]
+		]
+	];
+}
+
+void SAddTestPlanEntryDialog::OpenDialog()
+{
+	TSharedPtr<SWindow> PopupWindowPin = AddPointDialogWindow.Pin();
+	if (PopupWindowPin.IsValid())
+	{
+		PopupWindowPin->BringToFront();
+	}
+	else
+	{
+		PopupWindowPin = SNew(SWindow)
+			.HasCloseButton(true)
+			.SupportsMaximize(false)
+			.SupportsMinimize(false)
+			.ClientSize(FVector2D(240, 180));
+
+		FSlateApplication::Get().AddWindow(PopupWindowPin.ToSharedRef());
+	}
+
+	PopupWindowPin->SetTitle(LOCTEXT("AddTestPlanEntryDialog", "Add TestPlan item"));
+
+	TSharedPtr<SAddTestPlanEntryDialog> NewDialogContent = SNew(SAddTestPlanEntryDialog);
+	PopupWindowPin->SetContent(NewDialogContent.ToSharedRef());
+	AddPointDialogWindow = PopupWindowPin;
+}
+
+FReply SAddTestPlanEntryDialog::OnAddEntryClicked()
+{
+	IMessageBusTesterModule::Get().GetMessageBusTester().AddTestPlanItem(TestPlanItem);
+	CloseDialog();
+	return FReply::Handled();
+}
+
+FReply SAddTestPlanEntryDialog::OnCancelClicked()
+{
+	CloseDialog();
+	return FReply::Handled();
+}
+
+TSharedRef<SWidget> SAddTestPlanEntryDialog::MakeButtonsWidget()
+{
+	return SNew(SHorizontalBox)
+		+ SHorizontalBox::Slot()
+		[
+			SNew(SButton)
+			.ContentPadding(FAppStyle::GetMargin("StandardDialog.ContentPadding"))
+			.OnClicked(this, &SAddTestPlanEntryDialog::OnAddEntryClicked)
+			.HAlign(HAlign_Center)
+			.Text(LOCTEXT("AddDataPoint", "Add"))
+		]
+		+ SHorizontalBox::Slot()
+		[
+			SNew(SButton)
+			.ContentPadding(FAppStyle::GetMargin("StandardDialog.ContentPadding"))
+			.OnClicked(this, &SAddTestPlanEntryDialog::OnCancelClicked)
+			.HAlign(HAlign_Center)
+			.Text(LOCTEXT("CancelAddingDataPoint", "Cancel"))
+		];
+}
+
+void SAddTestPlanEntryDialog::CloseDialog()
+{
+	TSharedPtr<SWindow> ExistingWindowPin = AddPointDialogWindow.Pin();
+	if (ExistingWindowPin.IsValid())
+	{
+		ExistingWindowPin->RequestDestroyWindow();
+		ExistingWindowPin = nullptr;
+	}
+}
+
+
+#undef LOCTEXT_NAMESPACE 
diff --git a/Engine/Plugins/Experimental/MessageBusTester/Source/MessageBusTesterEditor/Private/Widgets/SAddTestPlanEntryDialog.h b/Engine/Plugins/Experimental/MessageBusTester/Source/MessageBusTesterEditor/Private/Widgets/SAddTestPlanEntryDialog.h
new file mode 100644
--- /dev/null
+++ b/Engine/Plugins/Experimental/MessageBusTester/Source/MessageBusTesterEditor/Private/Widgets/SAddTestPlanEntryDialog.h
@@ -0,0 +1,44 @@
+// Copyright Epic Games, Inc. All Rights Reserved.
+
+#pragma once
+
+#include "Widgets/SCompoundWidget.h"
+#include "MessageBusTesterCommon.h"
+#include "Widgets/DeclarativeSyntaxSupport.h"
+
+class SAddTestPlanEntryDialog : public SCompoundWidget
+{
+private:
+	using Super = SCompoundWidget;
+
+public:
+
+	SLATE_BEGIN_ARGS(SAddTestPlanEntryDialog)
+	{}
+
+	SLATE_END_ARGS()
+
+	void Construct(const FArguments& InArgs);
+
+	static void OpenDialog();
+
+private:
+
+	/** Called when Add button is clicked */
+	FReply OnAddEntryClicked();
+	
+	/** Called when Cancel button is clicked */
+	FReply OnCancelClicked();
+
+	TSharedRef<SWidget> MakeButtonsWidget();
+
+	/** Closes this dialog */
+	void CloseDialog();
+	
+	void AddEntry() const;
+
+private:
+	FTestPlanItem TestPlanItem;
+
+	static TWeakPtr<SWindow> AddPointDialogWindow;
+};
\ No newline at end of file
diff --git a/Engine/Plugins/Experimental/MessageBusTester/Source/MessageBusTesterEditor/Private/Widgets/SDiscoveredTesterListView.cpp b/Engine/Plugins/Experimental/MessageBusTester/Source/MessageBusTesterEditor/Private/Widgets/SDiscoveredTesterListView.cpp
new file mode 100644
--- /dev/null
+++ b/Engine/Plugins/Experimental/MessageBusTester/Source/MessageBusTesterEditor/Private/Widgets/SDiscoveredTesterListView.cpp
@@ -0,0 +1,423 @@
+// Copyright Epic Games, Inc. All Rights Reserved.
+
+#include "SDiscoveredTesterListView.h"
+
+#include "MessageBusTesterEditorModule.h"
+
+#include "EditorFontGlyphs.h"
+#include "EditorStyleSet.h"
+
+#include "Framework/Views/ITypedTableView.h"
+#include "IMessageBridge.h"
+#include "IMessageBusTester.h"
+#include "IMessageContext.h"
+
+#include "INetworkMessagingExtension.h"
+
+#include "Misc/App.h"
+#include "MessageBusTesterCommon.h"
+#include "IMessageBusTesterModule.h"
+#include "DiscoveredTester.h"
+
+#define LOCTEXT_NAMESPACE "SDiscoveredTesterListView"
+
+namespace DiscoveredTesterListView
+{
+	const FName HeaderIdName_ConnectionState = TEXT("ConnectionState");
+	const FName HeaderIdName_TestingState = TEXT("TestingState");
+	const FName HeaderIdName_MachineName = TEXT("MachineName");
+	const FName HeaderIdName_NetStats = TEXT("NetStats");
+	const FName HeaderIdName_ProcessId = TEXT("ProcessId");
+	const FName HeaderIdName_TesterName = TEXT("FriendlyName");
+	const FName HeaderIdName_ReliableKeepAlive = TEXT("ReliableKeepAlive");
+	const FName HeaderIdName_UnreliableKeepAlive = TEXT("UnreliableKeepAlive");
+}
+
+FDiscoveredTesterTableRowData::FDiscoveredTesterTableRowData(
+	const FGuid& InIdentifier,	const FTesterInstanceDescriptor& InDescriptor,	const FMessageAddress& Address)
+	: Identifier(InIdentifier)
+	, Descriptor(InDescriptor)
+{
+	if (INetworkMessagingExtension* NetworkMessagingExtension = FMessageBusTesterEditorModule::GetMessagingStatistics())
+	{
+		EndpointIdentifier = NetworkMessagingExtension->GetNodeIdFromAddress(Address);
+	}
+}
+
+INetworkMessagingExtension* GetMessagingStatistics()
+{
+	IModularFeatures& ModularFeatures = IModularFeatures::Get();
+	if (IsInGameThread())
+	{
+		if (ModularFeatures.IsModularFeatureAvailable(INetworkMessagingExtension::ModularFeatureName))
+		{
+			return &ModularFeatures.GetModularFeature<INetworkMessagingExtension>(INetworkMessagingExtension::ModularFeatureName);
+		}
+	}
+	else
+	{
+		ModularFeatures.LockModularFeatureList();
+		ON_SCOPE_EXIT
+		{
+			ModularFeatures.UnlockModularFeatureList();
+		};
+		
+		if (ModularFeatures.IsModularFeatureAvailable(INetworkMessagingExtension::ModularFeatureName))
+		{
+			return &ModularFeatures.GetModularFeature<INetworkMessagingExtension>(INetworkMessagingExtension::ModularFeatureName);
+		}
+	}
+	
+	ensureMsgf(false, TEXT("Feature %s is unavailable"), *INetworkMessagingExtension::ModularFeatureName.ToString());
+	return nullptr;
+}
+
+void FDiscoveredTesterTableRowData::UpdateCachedValues()
+{
+	const TConstArrayView<TSharedPtr<FDiscoveredTester, ESPMode::ThreadSafe>> DiscoveredTesters = IMessageBusTesterModule::Get().GetMessageBusTester().GetDiscoveredTesters();
+	const FGuid CurrentIdentifier = Identifier;
+	const TSharedPtr<FDiscoveredTester, ESPMode::ThreadSafe>* LinkedTesterPtr = DiscoveredTesters.FindByPredicate([CurrentIdentifier](const TSharedPtr<FDiscoveredTester, ESPMode::ThreadSafe>& Other) { return CurrentIdentifier == Other->Identifier; });
+	if (LinkedTesterPtr)
+	{
+		const TSharedPtr<FDiscoveredTester, ESPMode::ThreadSafe> LinkedTester = *LinkedTesterPtr;
+		CachedTestingState = LinkedTester->State;
+		CachedConnectionState = LinkedTester->ConnectionState;
+
+		CachedKeepAliveStatistics = LinkedTester->KeepAliveStatistics;
+	}
+	else
+	{
+		CachedTestingState = EMessageBusTesterState::Idle;
+		CachedConnectionState = EDiscoveredTesterConnectionState::Lost;
+	}
+
+	if (INetworkMessagingExtension* NetworkMessagingExtension = FMessageBusTesterEditorModule::GetMessagingStatistics())
+	{
+		Statistics = NetworkMessagingExtension->GetLatestNetworkStatistics(EndpointIdentifier);
+	}
+}
+
+FText FDiscoveredTesterTableRowData::GetNetSummary()
+{
+	return FText::Format(LOCTEXT("NetStats", "{0} Snt / {1} Rcv / {2} Ack"),
+						 FText::AsNumber(Statistics.PacketsSent),
+						 FText::AsNumber(Statistics.PacketsReceived),
+						 FText::AsNumber(Statistics.PacketsAcked));
+}
+
+/**
+ * SDiscoveredTesterTableRow
+ */
+void SDiscoveredTesterTableRow::Construct(const FArguments& InArgs, const TSharedRef<STableViewBase>& InOwerTableView)
+{
+	Item = InArgs._Item;
+	check(Item.IsValid());
+
+	Super::FArguments Arg;
+	Super::Construct(Arg, InOwerTableView);
+}
+
+TSharedRef<SWidget> SDiscoveredTesterTableRow::GenerateWidgetForColumn(const FName& ColumnName)
+{
+	if (DiscoveredTesterListView::HeaderIdName_ConnectionState == ColumnName)
+	{
+		return SNew(STextBlock)
+			.Font(FAppStyle::Get().GetFontStyle("FontAwesome.11"))
+			.Text(this, &SDiscoveredTesterTableRow::GetConnectionStateGlyphs)
+			.ColorAndOpacity(this, &SDiscoveredTesterTableRow::GetConnectionStateColorAndOpacity);
+	}
+	if (DiscoveredTesterListView::HeaderIdName_TestingState == ColumnName)
+	{
+		return SNew(STextBlock)
+			.Font(FAppStyle::Get().GetFontStyle("FontAwesome.11"))
+			.Text(this, &SDiscoveredTesterTableRow::GetTestingStateGlyphs)
+			.ColorAndOpacity(this, &SDiscoveredTesterTableRow::GetTestingStateColorAndOpacity);
+	}
+	if (DiscoveredTesterListView::HeaderIdName_MachineName == ColumnName)
+	{
+		return SNew(SHorizontalBox)
+			+ SHorizontalBox::Slot()
+			.HAlign(HAlign_Left)
+			.VAlign(VAlign_Center)
+			.AutoWidth()
+			[
+				SNew(STextBlock)
+				.Text(this, &SDiscoveredTesterTableRow::GetMachineName)
+			];
+	}
+	if (DiscoveredTesterListView::HeaderIdName_TesterName == ColumnName)
+	{
+		return SNew(SHorizontalBox)
+			+ SHorizontalBox::Slot()
+			.HAlign(HAlign_Left)
+			.VAlign(VAlign_Center)
+			.AutoWidth()
+			[
+				SNew(STextBlock)
+				.Text(this, &SDiscoveredTesterTableRow::GetFriendlyName)
+			];
+	}
+	if (DiscoveredTesterListView::HeaderIdName_ProcessId == ColumnName)
+	{
+		return SNew(SHorizontalBox)
+			+ SHorizontalBox::Slot()
+			.HAlign(HAlign_Left)
+			.VAlign(VAlign_Center)
+			.AutoWidth()
+			[
+				SNew(STextBlock)
+				.Text(this, &SDiscoveredTesterTableRow::GetProcessId)
+			];
+	}
+	if (DiscoveredTesterListView::HeaderIdName_NetStats == ColumnName)
+	{
+		return SNew(SHorizontalBox)
+			+ SHorizontalBox::Slot()
+			.HAlign(HAlign_Left)
+			.VAlign(VAlign_Center)
+			.AutoWidth()
+			[
+				SNew(STextBlock)
+				.Text(this, &SDiscoveredTesterTableRow::GetNetStats)
+			];
+	}
+	if (DiscoveredTesterListView::HeaderIdName_ReliableKeepAlive == ColumnName)
+	{
+		return SNew(SHorizontalBox)
+			+ SHorizontalBox::Slot()
+			.HAlign(HAlign_Left)
+			.VAlign(VAlign_Center)
+			.AutoWidth()
+			[
+				SNew(STextBlock)
+				.Text(this, &SDiscoveredTesterTableRow::GetReliableKeepAliveStats)
+			];
+	}
+	if (DiscoveredTesterListView::HeaderIdName_UnreliableKeepAlive == ColumnName)
+	{
+		return SNew(SHorizontalBox)
+			+ SHorizontalBox::Slot()
+			.HAlign(HAlign_Left)
+			.VAlign(VAlign_Center)
+			.AutoWidth()
+			[
+				SNew(STextBlock)
+				.Text(this, &SDiscoveredTesterTableRow::GetUnreliableKeepAliveStats)
+			];
+	}
+	
+	return SNullWidget::NullWidget;
+}
+
+FText SDiscoveredTesterTableRow::GetConnectionStateGlyphs() const
+{
+	return FEditorFontGlyphs::Circle;
+}
+
+FSlateColor SDiscoveredTesterTableRow::GetConnectionStateColorAndOpacity() const
+{
+	switch (Item->CachedConnectionState)
+	{
+		case EDiscoveredTesterConnectionState::Connected:
+		{
+			return FLinearColor::Green;
+		}
+		case EDiscoveredTesterConnectionState::Lost:
+		default:
+		{
+			return FLinearColor::Red;
+		}
+	}
+}
+
+FText SDiscoveredTesterTableRow::GetTestingStateGlyphs() const
+{
+	return FEditorFontGlyphs::Long_Arrow_Down;
+}
+
+FSlateColor SDiscoveredTesterTableRow::GetTestingStateColorAndOpacity() const
+{
+	switch (Item->CachedTestingState)
+	{
+	case EMessageBusTesterState::Active:
+	{
+		return FLinearColor::Green;
+	}
+	case EMessageBusTesterState::Idle:
+	default:
+	{
+		return FLinearColor::White;
+	}
+	}
+}
+
+FText SDiscoveredTesterTableRow::GetMachineName() const
+{
+	return FText::FromString(Item->Descriptor.MachineName);
+}
+
+FText SDiscoveredTesterTableRow::GetProcessId() const
+{
+	return FText::AsNumber(Item->Descriptor.ProcessId);
+}
+
+FText SDiscoveredTesterTableRow::GetNetStats() const
+{
+	return Item->GetNetSummary();
+}
+
+FText SDiscoveredTesterTableRow::GetReliableKeepAliveStats() const
+{
+	return FText::Format(LOCTEXT("ReliableKeepAliveStats", "{0} [{1} / {2} / {3}]"),
+		FText::AsNumber(Item->CachedKeepAliveStatistics.LastReliableKeepAliveReceptionInterval),
+		FText::AsNumber(Item->CachedKeepAliveStatistics.MinReliableKeepAliveInterval),
+		FText::AsNumber(Item->CachedKeepAliveStatistics.AverageReliableKeepAliveInterval),
+		FText::AsNumber(Item->CachedKeepAliveStatistics.MaxReliableKeepAliveInterval));
+}
+
+FText SDiscoveredTesterTableRow::GetUnreliableKeepAliveStats() const
+{
+	return FText::Format(LOCTEXT("UnreliableKeepAliveStats", "{0} [{1} / {2} / {3}]"),
+		FText::AsNumber(Item->CachedKeepAliveStatistics.LastUnreliableKeepAliveReceptionInterval),
+		FText::AsNumber(Item->CachedKeepAliveStatistics.MinUnreliableKeepAliveInterval),
+		FText::AsNumber(Item->CachedKeepAliveStatistics.AverageUnreliableKeepAliveInterval),
+		FText::AsNumber(Item->CachedKeepAliveStatistics.MaxUnreliableKeepAliveInterval));
+}
+
+FText SDiscoveredTesterTableRow::GetFriendlyName() const
+{
+	return FText::FromName(Item->Descriptor.FriendlyName);
+}
+
+void SDiscoveredTesterListView::HandleSelectionChanged(FDiscoveredTesterTableRowDataPtr Item, ESelectInfo::Type /*SelectInfo*/ )
+{
+	OnSelectionChanged().Broadcast(Item);
+}
+
+/**
+ * SDiscoveredTesterListView
+ */
+void SDiscoveredTesterListView::Construct(const FArguments& InArgs)
+{
+	IMessageBusTesterModule::Get().GetMessageBusTester().OnDiscoveredTesterListChanged().AddSP(this, &SDiscoveredTesterListView::OnDiscoveredTesterListChanged);
+
+	Super::Construct
+	(
+		Super::FArguments()
+		.ListItemsSource(&ListItemsSource)
+		.SelectionMode(ESelectionMode::Single)
+		.OnSelectionChanged(this, &SDiscoveredTesterListView::HandleSelectionChanged)
+		.OnGenerateRow(this, &SDiscoveredTesterListView::OnGenerateRow)
+		.HeaderRow
+		(
+			SNew(SHeaderRow)
+			.CanSelectGeneratedColumn(true) //To show/hide columns
+
+			+ SHeaderRow::Column(DiscoveredTesterListView::HeaderIdName_ConnectionState)
+			.FixedWidth(30.f)
+			.HAlignHeader(HAlign_Center)
+			.HAlignCell(HAlign_Center)
+			.DefaultLabel(LOCTEXT("HeaderName_ConnectionState", "ConnectionState"))
+
+			+ SHeaderRow::Column(DiscoveredTesterListView::HeaderIdName_TestingState)
+			.FixedWidth(30.f)
+			.HAlignHeader(HAlign_Center)
+			.HAlignCell(HAlign_Center)
+			.DefaultLabel(LOCTEXT("HeaderName_TestingState", "Testing"))
+
+			+ SHeaderRow::Column(DiscoveredTesterListView::HeaderIdName_MachineName)
+			.FillWidth(.1f)
+			.DefaultLabel(LOCTEXT("HeaderName_MachineName", "Machine"))
+
+			+ SHeaderRow::Column(DiscoveredTesterListView::HeaderIdName_ProcessId)
+			.FillWidth(.1f)
+			.DefaultLabel(LOCTEXT("HeaderName_ProcessId", "Process Id"))
+
+			+ SHeaderRow::Column(DiscoveredTesterListView::HeaderIdName_TesterName)
+			.FillWidth(.15f)
+			.DefaultLabel(LOCTEXT("HeaderName_FriendlyName", "Friendly Name"))
+			.ShouldGenerateWidget(true) //Can't hide this column
+
+			+ SHeaderRow::Column(DiscoveredTesterListView::HeaderIdName_UnreliableKeepAlive)
+			.FillWidth(.2f)
+			.DefaultLabel(LOCTEXT("HeaderName_UnreliableKeepAlive", "Unreliable KeepAlive"))
+
+			+ SHeaderRow::Column(DiscoveredTesterListView::HeaderIdName_ReliableKeepAlive)
+			.FillWidth(.2f)
+			.DefaultLabel(LOCTEXT("HeaderName_ReliableKeepAlive", "Reliable KeepAlive"))
+
+			+ SHeaderRow::Column(DiscoveredTesterListView::HeaderIdName_NetStats)
+			.FillWidth(.2f)
+			.DefaultLabel(LOCTEXT("HeaderName_NetStats", "Net stats"))
+		)
+	);
+
+	RebuildDiscoveredTesterList();
+}
+
+SDiscoveredTesterListView::~SDiscoveredTesterListView()
+{
+	if (IMessageBusTesterModule::IsAvailable())
+	{
+		IMessageBusTesterModule::Get().GetMessageBusTester().OnDiscoveredTesterListChanged().RemoveAll(this);
+	}
+}
+
+void SDiscoveredTesterListView::Tick(const FGeometry& AllottedGeometry, const double InCurrentTime, const float InDeltaTime)
+{
+	bool bForceRefresh = false;
+	if (bRebuildListRequested)
+	{
+		RebuildDiscoveredTesterList();
+		RebuildList();
+		bRebuildListRequested = false;
+		bForceRefresh = true;
+	}
+
+	//Update cached values at specific rate or when list has been rebuilt
+	constexpr double RefreshRate = .33;
+	if (bForceRefresh || (InCurrentTime - LastRefreshTime > RefreshRate))
+	{
+		LastRefreshTime = InCurrentTime;
+		for (FDiscoveredTesterTableRowDataPtr& RowDataPtr : ListItemsSource)
+		{
+			RowDataPtr->UpdateCachedValues();
+		}
+	}
+
+	Super::Tick(AllottedGeometry, InCurrentTime, InDeltaTime);
+}
+
+TSharedRef<ITableRow> SDiscoveredTesterListView::OnGenerateRow(FDiscoveredTesterTableRowDataPtr InItem, const TSharedRef<STableViewBase>& OwnerTable)
+{
+	TSharedRef<SDiscoveredTesterTableRow> Row = SNew(SDiscoveredTesterTableRow, OwnerTable).Item(InItem);
+	ListRowWidgets.Add(Row);
+	return Row;
+}
+
+void SDiscoveredTesterListView::OnDiscoveredTesterListChanged()
+{
+	bRebuildListRequested = true;
+}
+
+void SDiscoveredTesterListView::RebuildDiscoveredTesterList()
+{
+	ListItemsSource.Reset();
+
+	const TConstArrayView<TSharedPtr<FDiscoveredTester, ESPMode::ThreadSafe>> DiscoveredTesters = IMessageBusTesterModule::Get().GetMessageBusTester().GetDiscoveredTesters();
+	for (const TSharedPtr<FDiscoveredTester, ESPMode::ThreadSafe>& Tester : DiscoveredTesters)
+	{
+		TSharedRef<FDiscoveredTesterTableRowData> RowData = MakeShared<FDiscoveredTesterTableRowData>(Tester->Identifier, Tester->Descriptor, Tester->Address);
+		ListItemsSource.Add(RowData);
+	}
+
+	for (FDiscoveredTesterTableRowDataPtr& TableRowData : ListItemsSource)
+	{
+		TableRowData->UpdateCachedValues();
+	}
+
+	RequestListRefresh();
+}
+
+#undef LOCTEXT_NAMESPACE
+
diff --git a/Engine/Plugins/Experimental/MessageBusTester/Source/MessageBusTesterEditor/Private/Widgets/SDiscoveredTesterListView.h b/Engine/Plugins/Experimental/MessageBusTester/Source/MessageBusTesterEditor/Private/Widgets/SDiscoveredTesterListView.h
new file mode 100644
--- /dev/null
+++ b/Engine/Plugins/Experimental/MessageBusTester/Source/MessageBusTesterEditor/Private/Widgets/SDiscoveredTesterListView.h
@@ -0,0 +1,149 @@
+// Copyright Epic Games, Inc. All Rights Reserved.
+
+#pragma once
+
+#include "Widgets/Views/STableRow.h"
+#include "Widgets/Views/SListView.h"
+
+#include "CoreMinimal.h"
+#include "DiscoveredTester.h"
+#include "IMessageBusTesterModule.h"
+#include "INetworkMessagingExtension.h"
+#include "MessageBusTesterCommon.h"
+#include "SlateFwd.h"
+#include "Widgets/DeclarativeSyntaxSupport.h"
+
+
+
+
+class IMessageBusTester;
+/**
+ * FDiscoveredTesterTableRowData
+ */
+struct FDiscoveredTesterTableRowData : TSharedFromThis<FDiscoveredTesterTableRowData>
+{
+	FDiscoveredTesterTableRowData(const FGuid& InIdentifier, const FTesterInstanceDescriptor& InDescriptor, const FMessageAddress& Address);
+
+	/** Fetch latest information for the associated data provider */
+	void UpdateCachedValues();
+
+	/** Get a summary of current network statistics. */
+	FText GetNetSummary();
+
+	/** Return the current network statistics. */
+	const FMessageTransportStatistics& LatestStats()
+	{
+		return Statistics;
+	}
+
+public:
+
+	FGuid Identifier;
+	FGuid EndpointIdentifier;
+
+	FTesterInstanceDescriptor Descriptor;
+	FMessageTransportStatistics Statistics;
+	EDiscoveredTesterConnectionState CachedConnectionState;
+	EMessageBusTesterState CachedTestingState;
+	FKeepAliveStatistics CachedKeepAliveStatistics;
+};
+
+
+using FDiscoveredTesterTableRowDataPtr = TSharedPtr<FDiscoveredTesterTableRowData>;
+
+
+/**
+ *
+ */
+class SDiscoveredTesterTableRow : public SMultiColumnTableRow<FDiscoveredTesterTableRowDataPtr>
+{
+	using Super = SMultiColumnTableRow<FDiscoveredTesterTableRowDataPtr>;
+
+public:
+	SLATE_BEGIN_ARGS(SDiscoveredTesterTableRow) { }
+		SLATE_ARGUMENT(FDiscoveredTesterTableRowDataPtr, Item)
+	SLATE_END_ARGS()
+
+	void Construct(const FArguments& Args, const TSharedRef<STableViewBase>& OwerTableView);
+
+private:
+
+	/** Handles creation of each columns widget */
+	virtual TSharedRef<SWidget> GenerateWidgetForColumn(const FName& ColumnName) override;
+
+	/** Getters to populate the UI */
+	FText GetConnectionStateGlyphs() const;
+	FSlateColor GetConnectionStateColorAndOpacity() const;
+	FText GetTestingStateGlyphs() const;
+	FSlateColor GetTestingStateColorAndOpacity() const;
+	FText GetMachineName() const;
+	FText GetProcessId() const;
+	FText GetFriendlyName() const;
+	FText GetNetStats() const;
+	FText GetReliableKeepAliveStats() const;
+	FText GetUnreliableKeepAliveStats() const;
+
+private:
+
+	/** Item to display */
+	FDiscoveredTesterTableRowDataPtr Item;
+
+	/** Last time we refreshed the UI */
+	double LastRefreshTime = 0.0;
+};
+
+DECLARE_MULTICAST_DELEGATE_OneParam(FOnSelectionChangedNetworkItem, FDiscoveredTesterTableRowDataPtr);
+
+/**
+ *
+ */
+class SDiscoveredTesterListView : public SListView<FDiscoveredTesterTableRowDataPtr>
+{
+	using Super = SListView<FDiscoveredTesterTableRowDataPtr>;
+
+public:
+	SLATE_BEGIN_ARGS(SDiscoveredTesterListView) {}
+	SLATE_END_ARGS()
+
+	void Construct(const FArguments& InArgs);
+
+	/** Cleanup ourselves */
+	virtual ~SDiscoveredTesterListView();
+
+	/** Used to refresh cached values shown in UI */
+	virtual void Tick(const FGeometry& AllottedGeometry, const double InCurrentTime, const float InDeltaTime) override;
+
+	/**
+	 * Delegate that is called whenever the selection is changed on the network list.
+	 */
+	FOnSelectionChangedNetworkItem& OnSelectionChanged()
+	{
+		return SelectionChangeEvent;
+	}
+private:
+
+	/** Generate a new row for the listview using the Item data */
+	TSharedRef<ITableRow> OnGenerateRow(FDiscoveredTesterTableRowDataPtr Item, const TSharedRef<STableViewBase>& OwnerTable);
+
+	/** Handler for when the list view selection changes. */
+	void HandleSelectionChanged(FDiscoveredTesterTableRowDataPtr, ESelectInfo::Type /*SelectInfo*/ );
+
+	void OnDiscoveredTesterListChanged();
+
+	/** Rebuild provider list from scratch */
+	void RebuildDiscoveredTesterList();
+
+private:
+
+	TArray<FDiscoveredTesterTableRowDataPtr> ListItemsSource;
+	TArray<TWeakPtr<SDiscoveredTesterTableRow>> ListRowWidgets;
+
+	/** Used to cache if list needs to be refreshed or not */
+	bool bRebuildListRequested = false;
+
+	/** Timestamp when we last refreshed UI */
+	double LastRefreshTime = 0.0;
+
+	FOnSelectionChangedNetworkItem SelectionChangeEvent;
+};
+
diff --git a/Engine/Plugins/Experimental/MessageBusTester/Source/MessageBusTesterEditor/Private/Widgets/SMessageBusTestPlanListView.cpp b/Engine/Plugins/Experimental/MessageBusTester/Source/MessageBusTesterEditor/Private/Widgets/SMessageBusTestPlanListView.cpp
new file mode 100644
--- /dev/null
+++ b/Engine/Plugins/Experimental/MessageBusTester/Source/MessageBusTesterEditor/Private/Widgets/SMessageBusTestPlanListView.cpp
@@ -0,0 +1,194 @@
+// Copyright Epic Games, Inc. All Rights Reserved.
+
+#include "SMessageBusTestPlanListView.h"
+
+#include "EditorStyleSet.h"
+#include "Styling/AppStyle.h"
+
+#include "IMessageBusTester.h"
+#include "Misc/App.h"
+#include "MessageBusTesterCommon.h"
+#include "IMessageBusTesterModule.h"
+#include "Widgets/Input/SButton.h"
+#include "Widgets/Images/SImage.h"
+
+#define LOCTEXT_NAMESPACE "SMessageBusTestPlanListView"
+
+namespace MessageBusTestPlanListView
+{
+	const FName HeaderIdName_PayloadSize = TEXT("PayloadSize");
+	const FName HeaderIdName_Interval = TEXT("Interval");
+	const FName HeaderIdName_Remove = TEXT("Remove");
+}
+
+/**
+ * FMessageBusTestPlanTableRowData
+ */
+struct FMessageBusTestPlanTableRowData : TSharedFromThis<FMessageBusTestPlanTableRowData>
+{
+	FMessageBusTestPlanTableRowData(int32 InIndex, int32 InPayloadSizeBytes, double InIntervalSeconds)
+		: EntryIndex(InIndex)
+		, PayloadSizeBytes(InPayloadSizeBytes)
+		, IntervalSeconds(InIntervalSeconds)
+	{
+	}
+
+public:
+	
+	int32 EntryIndex = INDEX_NONE;
+	int32 PayloadSizeBytes = 0;
+	double IntervalSeconds = 0.0;
+};
+
+
+/**
+ * SMessageBusTestPlanTableRow
+ */
+void SMessageBusTestPlanTableRow::Construct(const FArguments& InArgs, const TSharedRef<STableViewBase>& InOwerTableView)
+{
+	Item = InArgs._Item;
+	check(Item.IsValid());
+
+	Super::FArguments Arg;
+	Super::Construct(Arg, InOwerTableView);
+}
+
+TSharedRef<SWidget> SMessageBusTestPlanTableRow::GenerateWidgetForColumn(const FName& ColumnName)
+{
+	if (MessageBusTestPlanListView::HeaderIdName_PayloadSize == ColumnName)
+	{
+		return SNew(SHorizontalBox)
+			+ SHorizontalBox::Slot()
+			.HAlign(HAlign_Left)
+			.VAlign(VAlign_Center)
+			.AutoWidth()
+			[
+				SNew(STextBlock)
+				.Text(this, &SMessageBusTestPlanTableRow::GetPayloadSize)
+			];
+	}
+	if (MessageBusTestPlanListView::HeaderIdName_Interval == ColumnName)
+	{
+		return SNew(SHorizontalBox)
+			+ SHorizontalBox::Slot()
+			.HAlign(HAlign_Left)
+			.VAlign(VAlign_Center)
+			.AutoWidth()
+			[
+				SNew(STextBlock)
+				.Text(this, &SMessageBusTestPlanTableRow::GetInterval)
+			];
+	}
+	if (MessageBusTestPlanListView::HeaderIdName_Remove == ColumnName)
+	{
+		return SNew(SButton)
+			.ButtonStyle(FAppStyle::Get(), "HoverHintOnly")
+			.HAlign(HAlign_Center)
+			.VAlign(VAlign_Center)
+			.OnClicked(this, &SMessageBusTestPlanTableRow::OnRemoveEntryClicked)
+			.ToolTipText(LOCTEXT("RemoveTestPlanEntry", "Remove entry from test plan"))
+			.ContentPadding(0.f)
+			.ForegroundColor(FSlateColor::UseForeground())
+			.IsFocusable(false)
+			[
+				SNew(SImage)
+				.Image(FAppStyle::GetBrush("PropertyWindow.Button_EmptyArray"))
+				.ColorAndOpacity(FSlateColor::UseForeground())
+			];
+	}
+	
+	return SNullWidget::NullWidget;
+}
+
+FReply SMessageBusTestPlanTableRow::OnRemoveEntryClicked()
+{
+	IMessageBusTesterModule::Get().GetMessageBusTester().RemoveTestPlanItem(Item->EntryIndex);
+	return FReply::Handled();
+}
+
+FText SMessageBusTestPlanTableRow::GetPayloadSize() const
+{
+	return FText::AsNumber(Item->PayloadSizeBytes);
+}
+
+FText SMessageBusTestPlanTableRow::GetInterval() const
+{
+	return FText::AsNumber(Item->IntervalSeconds);
+}
+
+/**
+ * SMessageBusTestPlanListView
+ */
+void SMessageBusTestPlanListView::Construct(const FArguments& InArgs)
+{
+	IMessageBusTesterModule::Get().GetMessageBusTester().OnTestPlanChanged().AddSP(this, &SMessageBusTestPlanListView::OnMessageBusTestPlanChanged);
+
+	Super::Construct
+	(
+		Super::FArguments()
+		.ListItemsSource(&ListItemsSource)
+		.SelectionMode(ESelectionMode::None)
+		.OnGenerateRow(this, &SMessageBusTestPlanListView::OnGenerateRow)
+		.HeaderRow
+		(
+			SNew(SHeaderRow)
+
+			+ SHeaderRow::Column(MessageBusTestPlanListView::HeaderIdName_PayloadSize)
+			.FillWidth(.6f)
+			.DefaultLabel(LOCTEXT("HeaderName_Payload", "Size (bytes)"))
+
+			+ SHeaderRow::Column(MessageBusTestPlanListView::HeaderIdName_Interval)
+			.FillWidth(.4f)
+			.DefaultLabel(LOCTEXT("HeaderName_Interval", "Interval (s)"))
+
+			+ SHeaderRow::Column(MessageBusTestPlanListView::HeaderIdName_Remove)
+			.HAlignHeader(HAlign_Center)
+			.HAlignCell(HAlign_Center)
+			.ManualWidth(20.f)
+			.DefaultLabel(LOCTEXT("HeaderName_ProcessId", ""))
+		)
+	);
+
+	RebuildMessageBusTestPlanList();
+}
+
+void SMessageBusTestPlanListView::Tick(const FGeometry& AllottedGeometry, const double InCurrentTime, const float InDeltaTime)
+{
+	if (bRebuildListRequested)
+	{
+		RebuildMessageBusTestPlanList();
+		RebuildList();
+		bRebuildListRequested = false;
+	}
+
+	Super::Tick(AllottedGeometry, InCurrentTime, InDeltaTime);
+}
+
+TSharedRef<ITableRow> SMessageBusTestPlanListView::OnGenerateRow(FMessageBusTestPlanTableRowDataPtr InItem, const TSharedRef<STableViewBase>& OwnerTable)
+{
+	TSharedRef<SMessageBusTestPlanTableRow> Row = SNew(SMessageBusTestPlanTableRow, OwnerTable)
+		.Item(InItem);
+	return Row;
+}
+
+void SMessageBusTestPlanListView::OnMessageBusTestPlanChanged()
+{
+	bRebuildListRequested = true;
+}
+
+void SMessageBusTestPlanListView::RebuildMessageBusTestPlanList()
+{
+	ListItemsSource.Reset();
+
+	const FMessageBusTestPlan& MessageBusTestPlan = IMessageBusTesterModule::Get().GetMessageBusTester().GetTestPlan();
+	for (int32 Index = 0; Index < MessageBusTestPlan.TestPlanItems.Num(); ++Index)
+	{
+		TSharedRef<FMessageBusTestPlanTableRowData> RowData = MakeShared<FMessageBusTestPlanTableRowData>(Index, MessageBusTestPlan.TestPlanItems[Index].NumBytes, MessageBusTestPlan.TestPlanItems[Index].IntervalSeconds);
+		ListItemsSource.Add(RowData);
+	}
+
+	RequestListRefresh();
+}
+
+#undef LOCTEXT_NAMESPACE
+
diff --git a/Engine/Plugins/Experimental/MessageBusTester/Source/MessageBusTesterEditor/Private/Widgets/SMessageBusTestPlanListView.h b/Engine/Plugins/Experimental/MessageBusTester/Source/MessageBusTesterEditor/Private/Widgets/SMessageBusTestPlanListView.h
new file mode 100644
--- /dev/null
+++ b/Engine/Plugins/Experimental/MessageBusTester/Source/MessageBusTesterEditor/Private/Widgets/SMessageBusTestPlanListView.h
@@ -0,0 +1,91 @@
+// Copyright Epic Games, Inc. All Rights Reserved.
+
+#pragma once
+
+#include "Widgets/Views/STableRow.h"
+#include "Widgets/Views/SListView.h"
+
+#include "CoreMinimal.h"
+#include "SlateFwd.h"
+#include "Widgets/DeclarativeSyntaxSupport.h"
+
+struct FMessageBusTestPlanTableRowData;
+class IMessageBusTester;
+
+
+using FMessageBusTestPlanTableRowDataPtr = TSharedPtr<FMessageBusTestPlanTableRowData>;
+
+
+/**
+ *	
+ */
+class SMessageBusTestPlanTableRow : public SMultiColumnTableRow<FMessageBusTestPlanTableRowDataPtr>
+{
+	using Super = SMultiColumnTableRow<FMessageBusTestPlanTableRowDataPtr>;
+
+public:
+	SLATE_BEGIN_ARGS(SMessageBusTestPlanTableRow) { }
+		SLATE_ARGUMENT(FMessageBusTestPlanTableRowDataPtr, Item)
+	SLATE_END_ARGS()
+
+	void Construct(const FArguments& Args, const TSharedRef<STableViewBase>& OwerTableView);
+
+private:
+
+	/** Handles creation of each columns widget */
+	virtual TSharedRef<SWidget> GenerateWidgetForColumn(const FName& ColumnName) override;
+
+	FReply OnRemoveEntryClicked();
+
+	/** Getters to populate the UI */
+	FText GetPayloadSize() const;
+	FText GetInterval() const;
+
+private:
+
+	/** Item to display */
+	FMessageBusTestPlanTableRowDataPtr Item;
+	
+	/** Last time we refreshed the UI */
+	double LastRefreshTime = 0.0;
+};
+
+
+/**
+ *
+ */
+class SMessageBusTestPlanListView : public SListView<FMessageBusTestPlanTableRowDataPtr>
+{
+	using Super = SListView<FMessageBusTestPlanTableRowDataPtr>;
+
+public:
+	SLATE_BEGIN_ARGS(SMessageBusTestPlanListView) {}
+	SLATE_END_ARGS()
+
+	void Construct(const FArguments& InArgs);
+
+	/** Used to refresh cached values shown in UI */
+	virtual void Tick(const FGeometry& AllottedGeometry, const double InCurrentTime, const float InDeltaTime) override;
+
+private:
+
+	/** Generate a new row for the listview using the Item data */
+	TSharedRef<ITableRow> OnGenerateRow(FMessageBusTestPlanTableRowDataPtr Item, const TSharedRef<STableViewBase>& OwnerTable);
+	
+	void OnMessageBusTestPlanChanged();
+	
+	/** Rebuild list from scratch */
+	void RebuildMessageBusTestPlanList();
+
+private:
+	
+	TArray<FMessageBusTestPlanTableRowDataPtr> ListItemsSource;
+	TArray<TWeakPtr<SMessageBusTestPlanTableRow>> ListRowWidgets;
+
+	/** Used to cache if list needs to be refreshed or not */
+	bool bRebuildListRequested = false;
+
+	/** Timestamp when we last refreshed UI */
+	double LastRefreshTime = 0.0;
+};
+
diff --git a/Engine/Plugins/Experimental/MessageBusTester/Source/MessageBusTesterEditor/Public/IMessageBusTesterEditorModule.h b/Engine/Plugins/Experimental/MessageBusTester/Source/MessageBusTesterEditor/Public/IMessageBusTesterEditorModule.h
new file mode 100644
--- /dev/null
+++ b/Engine/Plugins/Experimental/MessageBusTester/Source/MessageBusTesterEditor/Public/IMessageBusTesterEditorModule.h
@@ -0,0 +1,41 @@
+// Copyright Epic Games, Inc. All Rights Reserved.
+
+#pragma once
+
+#include "Modules/ModuleInterface.h"
+
+#include "CoreMinimal.h"
+#include "Features/IModularFeatures.h"
+#include "Misc/ScopeExit.h"
+#include "Misc/TVariant.h"
+#include "Modules/ModuleManager.h"
+
+class IMessageBusTester;
+class IMessageBusTesterLogger;
+
+class MESSAGEBUSTESTEREDITOR_API IMessageBusTesterEditorModule : public IModuleInterface
+{
+public:
+	/**
+	 * Singleton-like access to this module's interface.  This is just for convenience!
+	 * Beware of calling this during the shutdown phase, though.  Your module might have been unloaded already.
+	 *
+	 * @return Returns singleton instance, loading the module on demand if needed
+	 */
+	static IMessageBusTesterEditorModule& Get()
+	{
+		return FModuleManager::LoadModuleChecked<IMessageBusTesterEditorModule>("MessageBusTesterEditor");
+	}
+
+	/**
+	 * Checks to see if this module is loaded and ready.  It is only valid to call Get() during shutdown if IsAvailable() returns true.
+	 *
+	 * @return True if the module is loaded and ready to use
+	 */
+	static bool IsAvailable()
+	{
+		return FModuleManager::Get().IsModuleLoaded("MessageBusTesterEditor");
+	}
+
+	virtual void DisplayMessageBusTester() = 0;
+};
diff --git a/Engine/Source/Programs/LiveLinkHub/Source/LiveLinkHub.Target.cs b/Engine/Source/Programs/LiveLinkHub/Source/LiveLinkHub.Target.cs
--- a/Engine/Source/Programs/LiveLinkHub/Source/LiveLinkHub.Target.cs
+++ b/Engine/Source/Programs/LiveLinkHub/Source/LiveLinkHub.Target.cs
@@ -48,6 +48,7 @@ public class LiveLinkHubTarget : TargetRules
 			"PropertyAccessEditor",
 			"PythonScriptPlugin",
 			"UdpMessaging",
+			"MessageBusTester",
 			"CameraCalibrationCore",
 			"AppleARKitFaceSupport",
 			"XInputDevice",
