Add a start / stop button for enabling the MessageBus tester.  This will activate the tester keep alive messages so that it can start communicating with other message bus tester instances.

diff --git a/Engine/Plugins/Experimental/MessageBusTester/Source/MessageBusTester/Private/MessageBusTester.cpp b/Engine/Plugins/Experimental/MessageBusTester/Source/MessageBusTester/Private/MessageBusTester.cpp
--- a/Engine/Plugins/Experimental/MessageBusTester/Source/MessageBusTester/Private/MessageBusTester.cpp
+++ b/Engine/Plugins/Experimental/MessageBusTester/Source/MessageBusTester/Private/MessageBusTester.cpp
@@ -84,7 +84,12 @@ void FMessageBusTester::Initialize()
 
 	CachedDiscoveryMessage.Descriptor = MessageBusTesterUtils::GetInstanceDescriptor();
 
-	// Start();
+
+	// Only start the tester on initialization if the user has asked for it.
+	if (FParse::Param(FCommandLine::Get(), TEXT("MessageBusTesterStart")))
+	{
+		Start();
+	}
 }
 
 void FMessageBusTester::Start()
diff --git a/Engine/Plugins/Experimental/MessageBusTester/Source/MessageBusTester/Private/MessageBusTester.h b/Engine/Plugins/Experimental/MessageBusTester/Source/MessageBusTester/Private/MessageBusTester.h
--- a/Engine/Plugins/Experimental/MessageBusTester/Source/MessageBusTester/Private/MessageBusTester.h
+++ b/Engine/Plugins/Experimental/MessageBusTester/Source/MessageBusTester/Private/MessageBusTester.h
@@ -35,11 +35,27 @@ public:
 	{
 		return MakeArrayView(DiscoveredTesters);
 	}
+
+	virtual bool IsRunning() const override
+	{
+		return TesterEndpoint.IsValid();		
+	}
+	virtual bool StartSystem() override
+	{
+		Start();
+		return true;
+	}
+	virtual bool StopSystem() override
+	{
+		Stop();
+		return true;
+	}
 		
 	virtual bool ClearLostTesters() override;		
 	virtual const FMessageBusTestPlan& GetTestPlan() const override { return TestPlan; }
 	virtual void AddTestPlanItem(FTestPlanItem NewItem) override;
 	virtual void RemoveTestPlanItem(int32 Index) override;
+
 	virtual bool StartTest() override;
 	virtual bool StopTest() override;
 	virtual EMessageBusTesterState GetState() const override { return State; }
diff --git a/Engine/Plugins/Experimental/MessageBusTester/Source/MessageBusTester/Public/IMessageBusTester.h b/Engine/Plugins/Experimental/MessageBusTester/Source/MessageBusTester/Public/IMessageBusTester.h
--- a/Engine/Plugins/Experimental/MessageBusTester/Source/MessageBusTester/Public/IMessageBusTester.h
+++ b/Engine/Plugins/Experimental/MessageBusTester/Source/MessageBusTester/Public/IMessageBusTester.h
@@ -25,8 +25,13 @@ public:
 
 	/** Clear any lost testers and return true if any were removed. */
 	virtual bool ClearLostTesters() = 0;
-	
+
 	virtual const FMessageBusTestPlan& GetTestPlan() const = 0;
+
+	virtual bool IsRunning() const = 0;
+	virtual bool StartSystem() = 0;
+	virtual bool StopSystem() = 0;
+
 	virtual bool StartTest() = 0;
 	virtual bool StopTest() = 0;
 	virtual void AddTestPlanItem(FTestPlanItem NewItem) = 0;
diff --git a/Engine/Plugins/Experimental/MessageBusTester/Source/MessageBusTesterEditor/Private/SMessageBusTesterPanel.cpp b/Engine/Plugins/Experimental/MessageBusTester/Source/MessageBusTesterEditor/Private/SMessageBusTesterPanel.cpp
--- a/Engine/Plugins/Experimental/MessageBusTester/Source/MessageBusTesterEditor/Private/SMessageBusTesterPanel.cpp
+++ b/Engine/Plugins/Experimental/MessageBusTester/Source/MessageBusTesterEditor/Private/SMessageBusTesterPanel.cpp
@@ -6,6 +6,8 @@
 #include "EditorStyleSet.h"
 #include "Styling/AppStyle.h"
 
+#include "Math/Color.h"
+
 #include "Framework/Application/SlateApplication.h"
 #include "Framework/Docking/TabManager.h"
 #include "Framework/MultiBox/MultiBoxBuilder.h"
@@ -207,7 +209,6 @@ TSharedRef<SWidget> SMessageBusTesterPanel::MakeToolbarWidget()
 	.VAlign(VAlign_Center)
 	.Padding(FMargin(5.f, 5.f, 5.f, 5.f))
 	[
-
 		SNew(SHorizontalBox)
 		// Clear entries buttons
 		+ SHorizontalBox::Slot()
@@ -218,6 +219,24 @@ TSharedRef<SWidget> SMessageBusTesterPanel::MakeToolbarWidget()
 			SNew(SHorizontalBox)
 			+ SHorizontalBox::Slot()
 			.FillWidth(1.f)
+			[
+				SNew(SButton)
+				.ButtonStyle(FAppStyle::Get(), "ToggleButton")
+				.ToolTipText_Lambda([this]() {return GetStartStopButtonTooltip();})
+				.ContentPadding(FMargin(4.f, 4.f))
+				.HAlign(HAlign_Center)
+				.VAlign(VAlign_Center)
+				.OnClicked(this, &SMessageBusTesterPanel::OnStartStopClicked)
+				[
+					SNew(STextBlock)
+					.TextStyle(FAppStyle::Get(), "AnimViewport.MessageText")
+					.Font(FAppStyle::Get().GetFontStyle("FontAwesome.14"))
+					.Text(FEditorFontGlyphs::Stop)
+					.ColorAndOpacity_Lambda([this]() {return GetStartStopButtonColor();})
+				]
+			]
+			+ SHorizontalBox::Slot()
+			.FillWidth(1.f)
 			[
 				SNew(SButton)
 				.ButtonStyle(FAppStyle::Get(), "ToggleButton")
@@ -240,6 +259,42 @@ TSharedRef<SWidget> SMessageBusTesterPanel::MakeToolbarWidget()
 }
 
 END_SLATE_FUNCTION_BUILD_OPTIMIZATION
+FText SMessageBusTesterPanel::GetStartStopButtonTooltip()
+{
+	if (MessageBusTesterHelper::Get().GetMessageBusTester().IsRunning())
+	{
+		return LOCTEXT("StopTesting", "Stop testing framework.");
+	}
+	else
+	{
+		return LOCTEXT("StartTesting", "Start testing framework.");
+	}
+}
+
+FLinearColor SMessageBusTesterPanel::GetStartStopButtonColor()
+{
+	if (MessageBusTesterHelper::Get().GetMessageBusTester().IsRunning())
+	{
+		return FLinearColor::Green;
+	}
+	else
+	{
+		return FLinearColor::Red;
+	}
+}
+
+FReply SMessageBusTesterPanel::OnStartStopClicked()
+{
+	if (!MessageBusTesterHelper::Get().GetMessageBusTester().IsRunning())
+	{
+		MessageBusTesterHelper::Get().GetMessageBusTester().StartSystem();
+	}
+	else
+	{
+		MessageBusTesterHelper::Get().GetMessageBusTester().StopSystem();
+	}
+	return FReply::Handled();
+}
 
 FReply SMessageBusTesterPanel::OnClearClicked()
 {
diff --git a/Engine/Plugins/Experimental/MessageBusTester/Source/MessageBusTesterEditor/Private/SMessageBusTesterPanel.h b/Engine/Plugins/Experimental/MessageBusTester/Source/MessageBusTesterEditor/Private/SMessageBusTesterPanel.h
--- a/Engine/Plugins/Experimental/MessageBusTester/Source/MessageBusTesterEditor/Private/SMessageBusTesterPanel.h
+++ b/Engine/Plugins/Experimental/MessageBusTester/Source/MessageBusTesterEditor/Private/SMessageBusTesterPanel.h
@@ -18,7 +18,7 @@ class SMessageBusTestPlan;
 class SDiscoveredTesterListView;
 class FWorkspaceItem;
 class SMessageBusTestLogger;
-
+struct FLinearColor;
 
 /**
  * Main panel for Message Bus tester.
@@ -50,6 +50,15 @@ private:
 	/** Make the toolbar widgets */
 	TSharedRef<SWidget> MakeToolbarWidget();
 
+	/** Get the tooltip text for the stop / start toolbar button. */
+	FText GetStartStopButtonTooltip();
+
+	/** Get the start stop button color. */
+	FLinearColor GetStartStopButtonColor();
+		
+	/** Handles clear everything button clicked */
+	FReply OnStartStopClicked();
+		
 	/** Handles clear everything button clicked */
 	FReply OnClearClicked();
 
