Fix up comment in the tester panel to accurately reflect the functions purpose.

#rnx

diff --git a/Engine/Plugins/Experimental/MessageBusTester/Source/MessageBusTesterEditor/Private/SMessageBusTesterPanel.h b/Engine/Plugins/Experimental/MessageBusTester/Source/MessageBusTesterEditor/Private/SMessageBusTesterPanel.h
--- a/Engine/Plugins/Experimental/MessageBusTester/Source/MessageBusTesterEditor/Private/SMessageBusTesterPanel.h
+++ b/Engine/Plugins/Experimental/MessageBusTester/Source/MessageBusTesterEditor/Private/SMessageBusTesterPanel.h
@@ -56,7 +56,7 @@ private:
 	/** Get the start stop button color. */
 	FLinearColor GetStartStopButtonColor();
 		
-	/** Handles clear everything button clicked */
+	/** Handles start / stop button clicked */
 	FReply OnStartStopClicked();
 		
 	/** Handles clear everything button clicked */
