Move from an Async future object to a heart beat thread that we can start and stop depending on the tester state.   Moved the handling of messages off the game thread so that we can receive messages even if the game thread is blocked we "process" most of the messages on the game thread but some messages are processed as they are handled.  An inbound queue is established for receiving messages that need to be processed.

#rnx

diff --git a/Engine/Plugins/Experimental/MessageBusTester/Source/MessageBusTester/Private/MessageBusTester.cpp b/Engine/Plugins/Experimental/MessageBusTester/Source/MessageBusTester/Private/MessageBusTester.cpp
--- a/Engine/Plugins/Experimental/MessageBusTester/Source/MessageBusTester/Private/MessageBusTester.cpp
+++ b/Engine/Plugins/Experimental/MessageBusTester/Source/MessageBusTester/Private/MessageBusTester.cpp
@@ -9,7 +9,12 @@
 #include "Misc/CoreDelegates.h"
 #include "MessageBusTesterModule.h"
 #include "MessageBusTesterSettings.h"
-#include "Async/Async.h"
+
+#include "HAL/PlatformProcess.h"
+#include "HAL/Runnable.h"
+#include "HAL/RunnableThread.h"
+
+#include "Templates/Function.h"
 
 #include "Algo/RemoveIf.h"
 
@@ -19,13 +24,103 @@ namespace MessageBusTesterUtils
 	static constexpr int32 KeepAliveIntervalMilliSeconds = 1000;
 }
 
+class FMessageBusTesterHeartBeat : public FRunnable
+{
+public:
+	explicit FMessageBusTesterHeartBeat(TFunction<void()> InKeepAliveFunc)
+		: bIsRunning(false)
+		, bStopRequested(false)
+		, KeepAliveFunc(MoveTemp(InKeepAliveFunc))
+	{
+		Thread.Reset(FRunnableThread::Create(this, TEXT("MessageBusTester Heart Beat")));
+	}
+
+	~FMessageBusTesterHeartBeat()
+	{
+		if (Thread)
+		{
+			Thread->Kill(true);
+			Thread = nullptr;
+		}
+
+	}
+			
+	bool IsRunning() const
+	{
+		return Thread && bIsRunning;
+	}
+
+	//~ FRunnable Interface
+	virtual bool Init() override
+	{
+		bIsRunning = true;
+		return true;
+	}
+
+	/** 
+	 * Every second send out a keep alive message that other testers can see.
+	 */
+	virtual uint32 Run() override
+	{
+		while (!bStopRequested)
+		{
+			KeepAliveFunc();
+
+			FPlatformProcess::SleepNoStats(1.0f);
+		}
+		return 0;
+	}
+
+	/** Stop the running thread */
+	virtual void Stop() override
+	{
+		bStopRequested = true;
+	}
+
+	/** Thread was told to exit by the FRunnable interface. */
+	virtual void Exit() override
+	{
+		bIsRunning = false;
+	}
+	//~ FRunnable Interface.
+
+	/** Atomics for state tracking and external access to the runnable. */
+	std::atomic<bool> bIsRunning = false;
+	std::atomic<bool> bStopRequested = false;
+
+	TFunction<void()> KeepAliveFunc;
+	TUniquePtr<FRunnableThread> Thread;
+};
+
+FMessageBusTester::FMessageBusTester() = default;
 FMessageBusTester::~FMessageBusTester()
 {
 	FCoreDelegates::OnPreExit.RemoveAll(this);
+
+	if (TickerHandle.IsValid())
+	{
+		FTSTicker::GetCoreTicker().RemoveTicker(TickerHandle);
+		TickerHandle.Reset();
+	}
 	
 	Stop();
 }
 
+void FMessageBusTester::SetState(EMessageBusTesterState InState)
+{
+	State = InState;
+}
+
+bool FMessageBusTester::IsActive() const
+{
+	return GetState() == EMessageBusTesterState::Active;
+}
+
+EMessageBusTesterState FMessageBusTester::GetState() const
+{
+	return State;
+}
+
 void FMessageBusTester::AddTestPlanItem(FTestPlanItem NewItem)
 {
 	TestPlan.TestPlanItems.Emplace(MoveTemp(NewItem));
@@ -43,7 +138,7 @@ void FMessageBusTester::RemoveTestPlanItem(int32 Index)
 
 bool FMessageBusTester::StartTest()
 {
-	if (State == EMessageBusTesterState::Idle)
+	if (GetState() == EMessageBusTesterState::Idle)
 	{
 		SetupTestPlan();
 
@@ -60,13 +155,13 @@ bool FMessageBusTester::StartTest()
 
 bool FMessageBusTester::StopTest()
 {
-	if (State == EMessageBusTesterState::Active)
+	if (GetState() == EMessageBusTesterState::Active)
 	{
 		//Notify other testers we are stopping the test
 		FTesterStopTestMessage* NewMsg = FMessageEndpoint::MakeMessage<FTesterStopTestMessage>();
 		SendMessage(NewMsg, EMessageBusTesterMessageFlags::Reliable);
 
-		State = EMessageBusTesterState::Idle;
+		SetState(EMessageBusTesterState::Idle);
 
 		return true;
 	}
@@ -80,65 +175,101 @@ void FMessageBusTester::Initialize()
 
 	Identifier = FGuid::NewGuid();
 
-	KeepAliveEvent = FGenericPlatformProcess::GetSynchEventFromPool();
-
 	CachedDiscoveryMessage.Descriptor = MessageBusTesterUtils::GetInstanceDescriptor();
 
-
 	// Only start the tester on initialization if the user has asked for it.
 	if (FParse::Param(FCommandLine::Get(), TEXT("MessageBusTesterStart")))
 	{
 		Start();
 	}
+	
+	TickerHandle = FTSTicker::GetCoreTicker().AddTicker(FTickerDelegate::CreateRaw(this, &FMessageBusTester::Tick), 0.0f);
+}
+
+template<typename TInboundType>
+void FMessageBusTester::HandleMessage(const TInboundType& Inbound, const TSharedRef<IMessageContext, ESPMode::ThreadSafe>& Context)
+{
+	FInboundMessage InMsg(Inbound, Context);
+	InboundMessages.Enqueue(MoveTemp(InMsg));
+}
+
+template<>
+void FMessageBusTester::HandleMessage<FTesterKeepAliveMessage>(const FTesterKeepAliveMessage& Message, const TSharedRef<IMessageContext, ESPMode::ThreadSafe>& Context)
+{
+	FScopeLock Lock(&DiscoveredTesterCriticalSection);
+
+	const FGuid& SenderIdentifier = Message.Identifier;
+	if (TSharedPtr<FDiscoveredTester, ESPMode::ThreadSafe>* FoundTesterPtr = DiscoveredTesters.FindByPredicate([SenderIdentifier](const TSharedPtr<FDiscoveredTester, ESPMode::ThreadSafe>& Other) { return Other->Identifier == SenderIdentifier; }))
+	{
+		TSharedPtr<FDiscoveredTester, ESPMode::ThreadSafe>& FoundTester = *FoundTesterPtr;
+		if (FoundTester->ConnectionState != EDiscoveredTesterConnectionState::Lost)
+		{
+			FoundTester->NotifyKeepAliveReceived(Message.bReliablySent);
+			FoundTester->State = Message.State;
+		}
+	}
+}
+
+template<>
+void FMessageBusTester::HandleMessage<FTestPayloadMessage>(const FTestPayloadMessage& Message, const TSharedRef<IMessageContext, ESPMode::ThreadSafe>& Context)
+{
+	FTestPayloadReceptionMessage* Response = FMessageEndpoint::MakeMessage<FTestPayloadReceptionMessage>();
+	Response->PayloadId = Message.PayloadId;
+	Response->TestPlanIndex = Message.TestPlanIndex;
+	Response->ReceivedPayloadSize = Message.Payload.Num();
+	SendMessage(Response, {Context->GetSender()}, EMessageBusTesterMessageFlags::Reliable);
 }
 
 void FMessageBusTester::Start()
 {
 	UE_LOG(LogMessageBusTester, Log, TEXT("Starting MessageBusTester '%s'"), *CachedDiscoveryMessage.Descriptor.FriendlyName.ToString());
 
+	auto Handler = [this](const auto& T, const TSharedRef<IMessageContext, ESPMode::ThreadSafe>& Context)
+	{
+		HandleMessage(T, Context);
+	};
+	
 	const FString MessageEndpointName = TEXT("MessageBusTester");
 	TesterEndpoint = FMessageEndpoint::Builder(*MessageEndpointName)
-		.ReceivingOnThread(ENamedThreads::Type::GameThread)
-		.Handling<FTesterDiscoveryMessage>(this, &FMessageBusTester::HandleDiscoveryMessage)
-		.Handling<FTesterKeepAliveMessage>(this, &FMessageBusTester::HandleKeepAlive)
-		.Handling<FTesterStartTestMessage>(this, &FMessageBusTester::HandleStartTest)
-		.Handling<FTesterStopTestMessage>(this, &FMessageBusTester::HandleStopTest)
-		.Handling<FTestPayloadMessage>(this, &FMessageBusTester::HandlePayloadMessage)
-		.Handling<FTestPayloadReceptionMessage>(this, &FMessageBusTester::HandlePayloadReceptionMessage);
+		.ReceivingOnAnyThread()
+		.Handling<FTesterDiscoveryMessage>(Handler)
+		.Handling<FTesterKeepAliveMessage>(Handler)
+		.Handling<FTesterStartTestMessage>(Handler)
+		.Handling<FTesterStopTestMessage>(Handler)
+		.Handling<FTestPayloadMessage>(Handler)
+		.Handling<FTestPayloadReceptionMessage>(Handler);
 
 	if (TesterEndpoint.IsValid())
 	{
 		TesterEndpoint->Subscribe<FTesterDiscoveryMessage>();
 
-		bShouldSendKeepAlive = true;
-		KeepAliveTaskFuture = Async(EAsyncExecution::TaskGraph, [this]()
-		{ 
-			KeepAliveTask();
-		});
-
+		auto KeepAlive = [this](){SendKeepAlive_AnyThread();};
+		MessageBusHeartBeat = MakeUnique<FMessageBusTesterHeartBeat>(MoveTemp(KeepAlive));
 		AddTestPlanItem({ 10 * 1024 * 1024, 0.0f});
 	}
+}
 
-	TickerHandle = FTSTicker::GetCoreTicker().AddTicker(FTickerDelegate::CreateRaw(this, &FMessageBusTester::Tick), 0.0f);
+void FMessageBusTester::SendKeepAlive_AnyThread()
+{
+	FTesterKeepAliveMessage* NewKeepAlive = FMessageEndpoint::MakeMessage<FTesterKeepAliveMessage>(KeepAliveNumber++);
+	NewKeepAlive->State = GetState();
+	NewKeepAlive->bReliablySent = true;
+	SendMessage(NewKeepAlive, EMessageBusTesterMessageFlags::Reliable);
+	
+	NewKeepAlive = FMessageEndpoint::MakeMessage<FTesterKeepAliveMessage>(KeepAliveNumber++);
+	NewKeepAlive->State = GetState();
+	NewKeepAlive->bReliablySent = false;
+	SendMessage(NewKeepAlive, EMessageBusTesterMessageFlags::None);
 }
 
 void FMessageBusTester::Stop()
 {
-	if (KeepAliveEvent)
-	{
-		KeepAliveEvent->Trigger();
-	}
-
-	bShouldSendKeepAlive = false;
-	if (KeepAliveTaskFuture.IsValid())
-	{
-		KeepAliveTaskFuture.WaitFor(FTimespan::FromMilliseconds(100.0));
-	}
-
-	if (TickerHandle.IsValid())
+	StopTest();
+	
+	if (MessageBusHeartBeat.IsValid())
 	{
-		FTSTicker::GetCoreTicker().RemoveTicker(TickerHandle);
-		TickerHandle.Reset();
+		MessageBusHeartBeat->Stop();
+		MessageBusHeartBeat.Reset();
 	}
 
 	//Let providers know we're out
@@ -149,22 +280,34 @@ void FMessageBusTester::Stop()
 		TesterEndpoint.Reset();
 	}
 
-	InvalidTesters.Empty();
+	// Clear the inbound message queue.
+	while (TOptional<FInboundMessage> Msg = InboundMessages.Dequeue()) {}
 
-	if (KeepAliveEvent)
-	{
-		FGenericPlatformProcess::ReturnSynchEventToPool(KeepAliveEvent);
-		KeepAliveEvent = nullptr;
-	}
+	InvalidTesters.Empty();
 }
 
 bool FMessageBusTester::Tick(float DeltaTime)
 {
 	UpdateDiscoveredTester();
-	SendDiscoveryMessage();
+	
+	if (MessageBusHeartBeat.IsValid())
+	{
+		SendDiscoveryMessage();
+	}
 
-	if (State == EMessageBusTesterState::Active)
+	// Clear the inbound message queue.
+	while (TOptional<FInboundMessage> Msg = InboundMessages.Dequeue())
 	{
+		auto ProcessOne = [this, Context=Msg->Context](const auto& InboundMessage)
+		{
+			ProcessOnGameThread(InboundMessage, Context);
+		};
+		
+		Visit(ProcessOne, Msg->Type);
+	}
+
+	if (GetState() == EMessageBusTesterState::Active)
+	{		
 		UpdateTestSequence();
 	}
 
@@ -173,8 +316,6 @@ bool FMessageBusTester::Tick(float DeltaTime)
 
 void FMessageBusTester::UpdateDiscoveredTester()
 {
-	bool bTesterStateChanged = false;
-
 	FScopeLock Lock(&DiscoveredTesterCriticalSection);
 
 	for (TSharedPtr<FDiscoveredTester, ESPMode::ThreadSafe>& Tester : DiscoveredTesters)
@@ -256,40 +397,11 @@ void FMessageBusTester::UpdateTestSequence()
 						UE_LOG(LogMessageBusTester, Log, TEXT("Sending testplan item %d (payload %d of size %d)"), Index, PayloadId, Test.ItemDescription.NumBytes);
 					}
 				}
-				
 			}
 		}
 	}
 }
 
-void FMessageBusTester::KeepAliveTask()
-{
-	{
-		FScopeLock Lock(&DiscoveredTesterCriticalSection);
-
-		static uint64 KeepAliveNumber = 0;
-		FTesterKeepAliveMessage* NewKeepAlive = FMessageEndpoint::MakeMessage<FTesterKeepAliveMessage>(KeepAliveNumber++);
-		NewKeepAlive->State = State;
-		NewKeepAlive->bReliablySent = true;
-		SendMessage(NewKeepAlive, EMessageBusTesterMessageFlags::Reliable);
-
-		NewKeepAlive = FMessageEndpoint::MakeMessage<FTesterKeepAliveMessage>(KeepAliveNumber++);
-		NewKeepAlive->State = State;
-		NewKeepAlive->bReliablySent = false;
-		SendMessage(NewKeepAlive, EMessageBusTesterMessageFlags::None);
-	}
-
-	//Relaunch ourself if desired
-	if (bShouldSendKeepAlive)
-	{
-		KeepAliveTaskFuture = Async(EAsyncExecution::TaskGraph, [this]()
-		{
-			KeepAliveEvent->Wait(MessageBusTesterUtils::KeepAliveIntervalMilliSeconds);
-			KeepAliveTask();
-		});
-	}
-}
-
 void FMessageBusTester::SetupTestPlan()
 {
 	TestPlanItemTrackers.Empty(TestPlan.TestPlanItems.Num());
@@ -300,7 +412,7 @@ void FMessageBusTester::SetupTestPlan()
 		TestPlanItemTrackers.Emplace(MoveTemp(NewTest));
 	}
 
-	State = EMessageBusTesterState::Active;
+	SetState(EMessageBusTesterState::Active);
 }
 
 bool FMessageBusTester::SendMessageInternal(FTesterBaseMessage* Payload, UScriptStruct* Type, const TArray<FMessageAddress>& Addresses, EMessageBusTesterMessageFlags InFlags) const
@@ -319,49 +431,6 @@ bool FMessageBusTester::SendMessageInternal(FTesterBaseMessage* Payload, UScript
 	return false;
 }
 
-void FMessageBusTester::HandleDiscoveryMessage(const FTesterDiscoveryMessage& Message, const TSharedRef<IMessageContext, ESPMode::ThreadSafe>& Context)
-{
-	//Don't listen to ourselves
-	if (Message.Identifier == Identifier)
-	{
-		return;
-	}
-
-	TRACE_CPUPROFILER_EVENT_SCOPE(FMessageBusTester::HandleDiscoveryMessage);
-
-	if (Message.Descriptor.MessageVersion != MessageBusTesterUtils::SupportedMessageBusTesterVersion)
-	{
-		//Keep track of invalid providers to avoid log spamming
-		if (!InvalidTesters.Contains(Message.Identifier))
-		{
-			InvalidTesters.Add(Message.Identifier);
-
-			UE_LOG(LogMessageBusTester, Warning, TEXT("Received discovery response from Tester '%s' with unsupported version '%d'. Supported version is '%d'")
-				, *Message.Descriptor.FriendlyName.ToString()
-				, Message.Descriptor.MessageVersion
-				, MessageBusTesterUtils::SupportedMessageBusTesterVersion);
-		}
-
-		return;
-	}
-
-	const UMessageBusTesterSettings* Settings = GetDefault<UMessageBusTesterSettings>();
-	if (!Settings->bUseSessionId || Settings->GetSessionId() == Message.Descriptor.SessionId)
-	{
-		HandleDiscoveredTester(Message.Identifier, Message.Descriptor, Context->GetSender());
-	}
-	else if (!InvalidTesters.Contains(Message.Identifier))
-	{
-		//Keep track of invalid providers to avoid log spamming
-		InvalidTesters.Add(Message.Identifier);
-
-		UE_LOG(LogMessageBusTester, Warning, TEXT("Received discovery response from Tester '%s' with SessionId '%d' but expected '%d'")
-			, *Message.Descriptor.FriendlyName.ToString()
-			, Message.Descriptor.SessionId
-			, Settings->GetSessionId());
-	}
-}
-
 void FMessageBusTester::HandleDiscoveredTester(const FGuid& InIdentifier, const FTesterInstanceDescriptor& Descriptor, const FMessageAddress& Address)
 {
 	// -already known
@@ -451,44 +520,62 @@ void FMessageBusTester::FillTesterDescription(const FGuid& InIdentifier, const F
 	OutEntry->LastUnreliableKeepAliveMessageTime = OutEntry->LastReceivedMessageTime;
 }
 
-void FMessageBusTester::HandleKeepAlive(const FTesterKeepAliveMessage& Message,	const TSharedRef<IMessageContext, ESPMode::ThreadSafe>& Context)
+void FMessageBusTester::ProcessOnGameThread(const FTesterDiscoveryMessage& Message, const TSharedRef<IMessageContext, ESPMode::ThreadSafe>& Context)
 {
-	FScopeLock Lock(&DiscoveredTesterCriticalSection);
+	//Don't listen to ourselves
+	if (Message.Identifier == Identifier)
+	{
+		return;
+	}
 
-	const FGuid& SenderIdentifier = Message.Identifier;
-	if (TSharedPtr<FDiscoveredTester, ESPMode::ThreadSafe>* FoundTesterPtr = DiscoveredTesters.FindByPredicate([SenderIdentifier](const TSharedPtr<FDiscoveredTester, ESPMode::ThreadSafe>& Other) { return Other->Identifier == SenderIdentifier; }))
+	TRACE_CPUPROFILER_EVENT_SCOPE(FMessageBusTester::HandleDiscoveryMessage);
+
+	if (Message.Descriptor.MessageVersion != MessageBusTesterUtils::SupportedMessageBusTesterVersion)
 	{
-		TSharedPtr<FDiscoveredTester, ESPMode::ThreadSafe>& FoundTester = *FoundTesterPtr;
-		if (FoundTester->ConnectionState != EDiscoveredTesterConnectionState::Lost)
+		//Keep track of invalid providers to avoid log spamming
+		if (!InvalidTesters.Contains(Message.Identifier))
 		{
-			FoundTester->NotifyKeepAliveReceived(Message.bReliablySent);
-			FoundTester->State = Message.State;
+			InvalidTesters.Add(Message.Identifier);
+
+			UE_LOG(LogMessageBusTester, Warning, TEXT("Received discovery response from Tester '%s' with unsupported version '%d'. Supported version is '%d'")
+				, *Message.Descriptor.FriendlyName.ToString()
+				, Message.Descriptor.MessageVersion
+				, MessageBusTesterUtils::SupportedMessageBusTesterVersion);
 		}
+
+		return;
+	}
+
+	const UMessageBusTesterSettings* Settings = GetDefault<UMessageBusTesterSettings>();
+	if (!Settings->bUseSessionId || Settings->GetSessionId() == Message.Descriptor.SessionId)
+	{
+		HandleDiscoveredTester(Message.Identifier, Message.Descriptor, Context->GetSender());
+	}
+	else if (!InvalidTesters.Contains(Message.Identifier))
+	{
+		//Keep track of invalid providers to avoid log spamming
+		InvalidTesters.Add(Message.Identifier);
+
+		UE_LOG(LogMessageBusTester, Warning, TEXT("Received discovery response from Tester '%s' with SessionId '%d' but expected '%d'")
+			, *Message.Descriptor.FriendlyName.ToString()
+			, Message.Descriptor.SessionId
+			, Settings->GetSessionId());
 	}
 }
 
-void FMessageBusTester::HandleStartTest(const FTesterStartTestMessage& Message,	const TSharedRef<IMessageContext, ESPMode::ThreadSafe>& Context)
+void FMessageBusTester::ProcessOnGameThread(const FTesterStartTestMessage& Message,	const TSharedRef<IMessageContext, ESPMode::ThreadSafe>& Context)
 {
 	TestPlan = Message.TestPlan;
 	OnTestPlanChanged().Broadcast();
 	SetupTestPlan();
 }
 
-void FMessageBusTester::HandleStopTest(const FTesterStopTestMessage& Message, const TSharedRef<IMessageContext, ESPMode::ThreadSafe>& Context)
+void FMessageBusTester::ProcessOnGameThread(const FTesterStopTestMessage& Message, const TSharedRef<IMessageContext, ESPMode::ThreadSafe>& Context)
 {
-	State = EMessageBusTesterState::Idle;
+	SetState(EMessageBusTesterState::Idle);
 }
 
-void FMessageBusTester::HandlePayloadMessage(const FTestPayloadMessage& Message, const TSharedRef<IMessageContext, ESPMode::ThreadSafe>& Context)
-{
-	FTestPayloadReceptionMessage* Response = FMessageEndpoint::MakeMessage<FTestPayloadReceptionMessage>();
-	Response->PayloadId = Message.PayloadId;
-	Response->TestPlanIndex = Message.TestPlanIndex;
-	Response->ReceivedPayloadSize = Message.Payload.Num();
-	SendMessage(Response,  { Context->GetSender() }, EMessageBusTesterMessageFlags::Reliable);
-}
-
-void FMessageBusTester::HandlePayloadReceptionMessage(const FTestPayloadReceptionMessage& Message, const TSharedRef<IMessageContext, ESPMode::ThreadSafe>& Context)
+void FMessageBusTester::ProcessOnGameThread(const FTestPayloadReceptionMessage& Message, const TSharedRef<IMessageContext, ESPMode::ThreadSafe>& Context)
 {
 	if (TestPlanItemTrackers.IsValidIndex(Message.TestPlanIndex))
 	{
@@ -502,6 +589,8 @@ void FMessageBusTester::HandlePayloadReceptionMessage(const FTestPayloadReceptio
 
 bool FMessageBusTester::ClearLostTesters()
 {
+	FScopeLock Lock(&DiscoveredTesterCriticalSection);
+	
 	int32 OldSize = DiscoveredTesters.Num();
 	int32 NewSize = Algo::RemoveIf(DiscoveredTesters, [](const TSharedPtr<FDiscoveredTester>& Tester) {
 		return Tester.IsValid() && Tester->ConnectionState == EDiscoveredTesterConnectionState::Lost;
diff --git a/Engine/Plugins/Experimental/MessageBusTester/Source/MessageBusTester/Private/MessageBusTester.h b/Engine/Plugins/Experimental/MessageBusTester/Source/MessageBusTester/Private/MessageBusTester.h
--- a/Engine/Plugins/Experimental/MessageBusTester/Source/MessageBusTester/Private/MessageBusTester.h
+++ b/Engine/Plugins/Experimental/MessageBusTester/Source/MessageBusTester/Private/MessageBusTester.h
@@ -13,7 +13,14 @@
 #include <atomic>
 #include "Async/Future.h"
 
+#include "Containers/MpscQueue.h"
+#include "Misc/TVariant.h"
+#include "Templates/UniquePtr.h"
+
 class FMessageEndpoint;
+class FMessageBusTesterHeartBeat;
+
+using FMessageBusTesterTypes = TVariant<FTesterDiscoveryMessage, FTesterStartTestMessage, FTesterStopTestMessage, FTestPayloadReceptionMessage>;
 
 /**
  * Implementation of the MessageBus Tester
@@ -21,14 +28,26 @@ class FMessageEndpoint;
 class FMessageBusTester : public IMessageBusTester
 {
 public:
+	struct FInboundMessage
+	{
+		template<typename T>
+		FInboundMessage(const T& InType, const TSharedRef<IMessageContext, ESPMode::ThreadSafe>& InContext)
+			: Context(InContext)
+		{
+			Type.Set<T>(InType);
+		}
 
-	FMessageBusTester() = default;
+		FMessageBusTesterTypes Type;
+		TSharedRef<IMessageContext, ESPMode::ThreadSafe> Context;
+	};
+		
+	FMessageBusTester();
 	FMessageBusTester(const FMessageBusTester&) = delete;
 	FMessageBusTester& operator=(const FMessageBusTester&) = delete;
 	virtual ~FMessageBusTester();
 
 	//~Begin IMessageBusTester interface
-	virtual bool IsActive() const override { return State == EMessageBusTesterState::Active; }
+	virtual bool IsActive() const override;
 	virtual FOnDiscoveredTesterListChanged& OnDiscoveredTesterListChanged() override { return DiscoveredTesterListChangedDelegate; }
 	virtual FOnTestPlanChanged& OnTestPlanChanged() override { return TestPlanChangedDelegate; }
 	virtual TConstArrayView<TSharedPtr<FDiscoveredTester, ESPMode::ThreadSafe>> GetDiscoveredTesters() const override
@@ -58,7 +77,7 @@ public:
 
 	virtual bool StartTest() override;
 	virtual bool StopTest() override;
-	virtual EMessageBusTesterState GetState() const override { return State; }
+	virtual EMessageBusTesterState GetState() const override;
 	//~End IMessageBusTester
 
 
@@ -67,22 +86,22 @@ public:
 	void Stop();
 
 private:
+	friend class FMessageBusTesterHeartBeat;
 
+	void SendKeepAlive_AnyThread();
+		
 	/** Sends a message using messagebus to connected providers */
 	bool SendMessageInternal(FTesterBaseMessage* Payload, UScriptStruct* Type, const TArray<FMessageAddress>& Addresses, EMessageBusTesterMessageFlags InFlags) const;
 
-	void HandleDiscoveryMessage(const FTesterDiscoveryMessage& Message, const TSharedRef<IMessageContext, ESPMode::ThreadSafe>& Context);
-	void HandleDiscoveredTester(const FGuid& Identifier, const FTesterInstanceDescriptor& Descriptor,
-	                            const FMessageAddress& Address);
-	bool AddTester(const FGuid& Identifier, const FTesterInstanceDescriptor& Descriptor,
-	               const FMessageAddress& Address);
-	void FillTesterDescription(const FGuid& Identifier, const FTesterInstanceDescriptor& NewDescriptor,
-	                           const FMessageAddress& NewAddress, TSharedPtr<FDiscoveredTester, ESPMode::ThreadSafe>& OutEntry) const;
-	void HandleKeepAlive(const FTesterKeepAliveMessage& Message, const TSharedRef<IMessageContext, ESPMode::ThreadSafe>& Context);
-	void HandleStartTest(const FTesterStartTestMessage& Message, const TSharedRef<IMessageContext, ESPMode::ThreadSafe>& Context);
-	void HandleStopTest(const FTesterStopTestMessage& Message, const TSharedRef<IMessageContext, ESPMode::ThreadSafe>& Context);
-	void HandlePayloadMessage(const FTestPayloadMessage& Message, const TSharedRef<IMessageContext, ESPMode::ThreadSafe>& Context);
-	void HandlePayloadReceptionMessage(const FTestPayloadReceptionMessage& Message, const TSharedRef<IMessageContext, ESPMode::ThreadSafe>& Context);
+
+	void HandleDiscoveredTester(const FGuid& Identifier, const FTesterInstanceDescriptor& Descriptor, const FMessageAddress& Address);
+	bool AddTester(const FGuid& Identifier, const FTesterInstanceDescriptor& Descriptor, const FMessageAddress& Address);
+	void FillTesterDescription(const FGuid& Identifier, const FTesterInstanceDescriptor& NewDescriptor, const FMessageAddress& NewAddress, TSharedPtr<FDiscoveredTester, ESPMode::ThreadSafe>& OutEntry) const;
+		
+	void ProcessOnGameThread(const FTesterDiscoveryMessage& Message, const TSharedRef<IMessageContext, ESPMode::ThreadSafe>& Context);
+	void ProcessOnGameThread(const FTesterStartTestMessage& Message, const TSharedRef<IMessageContext, ESPMode::ThreadSafe>& Context);
+	void ProcessOnGameThread(const FTesterStopTestMessage& Message, const TSharedRef<IMessageContext, ESPMode::ThreadSafe>& Context);
+	void ProcessOnGameThread(const FTestPayloadReceptionMessage& Message, const TSharedRef<IMessageContext, ESPMode::ThreadSafe>& Context);
 	
 	/**  */
 	bool Tick(float DeltaTime);
@@ -126,20 +145,35 @@ private:
 		return SendMessageInternal(Payload, MessageType::StaticStruct(), Addresses, Flags);
 	}
 
-	void KeepAliveTask();
-
 	void SetupTestPlan();
 
+	template<typename TInboundType>
+	void HandleMessage(const TInboundType& Inbound, const TSharedRef<IMessageContext, ESPMode::ThreadSafe>& Context);
+
+	void SetState(EMessageBusTesterState InState);
+
 private:
+	/** Messages received from message bus. */
+	TMpscQueue<FInboundMessage> InboundMessages;
+
+	/** Unique pointer to the heart beat thread runnable. */
+	TUniquePtr<FMessageBusTesterHeartBeat> MessageBusHeartBeat;
 
 	/** Endpoint used to communicate with other testers */
 	TSharedPtr<FMessageEndpoint, ESPMode::ThreadSafe> TesterEndpoint;
 
+	/** List of discovered endpoints that we are tracking for testing. */
 	TArray<TSharedPtr<FDiscoveredTester, ESPMode::ThreadSafe>> DiscoveredTesters;
 
 	/** List of providers we received discovery response from that were not usable (wrong version, session id...) */
 	TArray<FGuid> InvalidTesters;
 
+	/** Current state of the tester. */
+	std::atomic<EMessageBusTesterState> State = EMessageBusTesterState::Idle;
+
+	/** Tracks the current keep alive counter. */
+	std::atomic<uint64> KeepAliveNumber = 0;
+		
 	/** Timestamp when we last sent a discovery message */
 	double LastSentDiscoveryMessage = 0.0;
 
@@ -152,15 +186,8 @@ private:
 	/** DiscoveryMessage built once and broadcasted periodically */
 	FTesterDiscoveryMessage CachedDiscoveryMessage;
 
-	EMessageBusTesterState State = EMessageBusTesterState::Idle;
-
 	FCriticalSection DiscoveredTesterCriticalSection;
 
-	std::atomic<bool> bShouldSendKeepAlive = false;
-		
-	TFuture<void> KeepAliveTaskFuture;
-	FEvent* KeepAliveEvent = nullptr;
-
 	FOnDiscoveredTesterListChanged DiscoveredTesterListChangedDelegate;
 	FOnTestPlanChanged TestPlanChangedDelegate;
 
