Add some additional command line options to help with automating test plans with Message Bus Tester.

- Adds a command line flag to specify test plans (-MESSAGE_BUS_TESTER_PLANS).  Plans are specified using a <size>:<interval> pattern. For example, 10000:1.0 will send 10,000 bytes every 1 second.  You can have multiple plans by adding a comma, e.g. 1024:0.0,100000:1.0.
- Adds a command line flag to write results to csv (-MESSAGE_BUS_TESTER_EXPORT_TO_CSV).  The file will get overwritten each run.
- Adds a command line flag to specify test duration (-MESSAGE_BUS_TESTER_RUN_DURATION).   Indicates how long to run the test.  At the end of the period it will request an engine shutdown.
- Renames command line flag to start on launch (-MESSAGE_BUS_TESTER_START_ON_LAUNCH).  We have deprecated the -MessageBusTesterStart flag

diff --git a/Engine/Plugins/Animation/LiveLinkHub/Source/LiveLinkHubEditor/Private/LiveLinkHubEditorModule.cpp b/Engine/Plugins/Animation/LiveLinkHub/Source/LiveLinkHubEditor/Private/LiveLinkHubEditorModule.cpp
--- a/Engine/Plugins/Animation/LiveLinkHub/Source/LiveLinkHubEditor/Private/LiveLinkHubEditorModule.cpp
+++ b/Engine/Plugins/Animation/LiveLinkHub/Source/LiveLinkHubEditor/Private/LiveLinkHubEditorModule.cpp
@@ -53,7 +53,7 @@ void FLiveLinkHubEditorModule::ShutdownModule()
 
 void FLiveLinkHubEditorModule::OnPostEngineInit()
 {
-	if (GEditor)
+	if (GEditor && FApp::CanEverRender())
 	{
 		RegisterLiveLinkHubStatusBar();
 		
diff --git a/Engine/Plugins/Experimental/MessageBusTester/Source/MessageBusTester/Private/DiscoveredTester.cpp b/Engine/Plugins/Experimental/MessageBusTester/Source/MessageBusTester/Private/DiscoveredTester.cpp
--- a/Engine/Plugins/Experimental/MessageBusTester/Source/MessageBusTester/Private/DiscoveredTester.cpp
+++ b/Engine/Plugins/Experimental/MessageBusTester/Source/MessageBusTester/Private/DiscoveredTester.cpp
@@ -6,6 +6,12 @@
 #include "IMessageBusTesterLogger.h"
 #include "MessageBusTesterSettings.h"
 
+#include "Modules/ModuleManager.h"
+#include "Modules/ModuleInterface.h"
+#include "Features/IModularFeatures.h"
+
+#include "INetworkMessagingExtension.h"
+
 namespace UE::DiscoveredTester::Private
 {
 	inline IMessageBusTesterModule& Get()
@@ -46,6 +52,14 @@ void FDiscoveredTester::UpdateConnectionState()
 			ConnectionState = EDiscoveredTesterConnectionState::Connected;
 		}
 	}
+
+	IModularFeatures& ModularFeatures = IModularFeatures::Get();
+	if (ModularFeatures.IsModularFeatureAvailable(INetworkMessagingExtension::ModularFeatureName))
+	{
+		INetworkMessagingExtension& Network = ModularFeatures.GetModularFeature<INetworkMessagingExtension>(INetworkMessagingExtension::ModularFeatureName);
+		FGuid Id = Network.GetNodeIdFromAddress(Address);
+		Statistics = Network.GetLatestNetworkStatistics(Id);
+	}
 }
 
 void FDiscoveredTester::NotifyPayloadSent(const FTestPayloadMessage& PayloadMessage)
@@ -64,8 +78,12 @@ void FDiscoveredTester::NotifyAckReceived(const FTestPayloadReceptionMessage& Ac
 		const double Elapsed = FMath::Max(ReceivedTime - ReceivedPayload.SentTime, UE_KINDA_SMALL_NUMBER);
 		const double MegaByte = 1024.0f * 1024.0f;
 		const double BandwithMegaBytes = static_cast<double>(ReceivedPayload.PayloadSize) / MegaByte / Elapsed;
-		UE::DiscoveredTester::Private::Get().GetLogger().Log(Descriptor.FriendlyName, FString::Printf(TEXT("Packet %d acknowledged after %0.5f (%0.1f MB/s)"), ReceivedPayload.PayloadId, Elapsed, BandwithMegaBytes), EMessageSeverity::Info);
+
+		// Compute running average. 
+		AverageMbPerSecond.PushValue(BandwithMegaBytes);
 		
+		UE::DiscoveredTester::Private::Get().GetLogger().Log(Descriptor.FriendlyName, FString::Printf(TEXT("Packet %d acknowledged after %0.5f (%0.1f MB/s) - Avg (%0.1f MB/s)"), ReceivedPayload.PayloadId, Elapsed, BandwithMegaBytes, AverageMbPerSecond.GetAverage()), EMessageSeverity::Info);
+
 		ReceivedPayload.AcknowledgedTime = ReceivedTime;
 		SentPayloads.RemoveAtSwap(SentPayloadIndex);
 		AcknowledgedPayloads.AddUnique(MoveTemp(ReceivedPayload));
diff --git a/Engine/Plugins/Experimental/MessageBusTester/Source/MessageBusTester/Private/MessageBusTester.cpp b/Engine/Plugins/Experimental/MessageBusTester/Source/MessageBusTester/Private/MessageBusTester.cpp
--- a/Engine/Plugins/Experimental/MessageBusTester/Source/MessageBusTester/Private/MessageBusTester.cpp
+++ b/Engine/Plugins/Experimental/MessageBusTester/Source/MessageBusTester/Private/MessageBusTester.cpp
@@ -7,12 +7,18 @@
 #include "MessageEndpointBuilder.h"
 #include "Misc/App.h"
 #include "Misc/CoreDelegates.h"
+#include "Misc/Paths.h"
+#include "Misc/DateTime.h"
+
 #include "MessageBusTesterModule.h"
 #include "MessageBusTesterSettings.h"
 
+#include "CoreGlobals.h"
+
 #include "HAL/PlatformProcess.h"
 #include "HAL/Runnable.h"
 #include "HAL/RunnableThread.h"
+#include "HAL/FileManager.h"
 
 #include "Templates/Function.h"
 
@@ -108,6 +114,11 @@ FMessageBusTester::~FMessageBusTester()
 
 void FMessageBusTester::SetState(EMessageBusTesterState InState)
 {
+	if (InState == State)
+	{
+		return;
+	}
+
 	State = InState;
 }
 
@@ -123,6 +134,7 @@ EMessageBusTesterState FMessageBusTester::GetState() const
 
 void FMessageBusTester::AddTestPlanItem(FTestPlanItem NewItem)
 {
+	UE_LOG(LogMessageBusTester, Display, TEXT("Adding test plan size %d with interval %f."), NewItem.NumBytes, NewItem.IntervalSeconds);
 	TestPlan.TestPlanItems.Emplace(MoveTemp(NewItem));
 	OnTestPlanChanged().Broadcast();
 }
@@ -153,16 +165,27 @@ bool FMessageBusTester::StartTest()
 	return false;
 }
 
-bool FMessageBusTester::StopTest()
+void FMessageBusTester::TriggerAppExit()
+{
+	const double ShutdownWaitTime = 3.0;
+	ShutdownTime = FPlatformTime::Seconds() + ShutdownWaitTime;
+}
+
+bool FMessageBusTester::StopTest(bool bShouldExitOnStop)
 {
 	if (GetState() == EMessageBusTesterState::Active)
 	{
 		//Notify other testers we are stopping the test
 		FTesterStopTestMessage* NewMsg = FMessageEndpoint::MakeMessage<FTesterStopTestMessage>();
+		NewMsg->bEngineShouldExit = bShouldExitOnStop; 
 		SendMessage(NewMsg, EMessageBusTesterMessageFlags::Reliable);
 
-		SetState(EMessageBusTesterState::Idle);
+		if (bShouldExitOnStop)
+		{
+			TriggerAppExit();
+		}
 
+		SetState(EMessageBusTesterState::Idle);
 		return true;
 	}
 
@@ -177,9 +200,64 @@ void FMessageBusTester::Initialize()
 
 	CachedDiscoveryMessage.Descriptor = MessageBusTesterUtils::GetInstanceDescriptor();
 
+	const TCHAR* CommandLine = FCommandLine::Get();
+
+	// Parse test plans to run.
+	FString TestPlans;
+	FParse::Value(CommandLine, TEXT("-MESSAGE_BUS_TESTER_PLANS="), TestPlans, false);
+
+	TArray<FString> CommandLinePlans;
+	TestPlans.ParseIntoArray(CommandLinePlans, TEXT(","));
+	for (const FString& Plan : CommandLinePlans)
+	{
+		TArray<FString> Tokens;
+		if (Plan.ParseIntoArray(Tokens, TEXT(":"), false) == 2)
+		{
+			const int32 PlanSize = FCString::Atoi(*Tokens[0]);
+			if (PlanSize > 0)
+			{
+				const float Interval = FCString::Atof(*Tokens[1]);
+				if (Interval >= 0.0)
+				{
+					AddTestPlanItem({PlanSize, Interval});
+				}
+				else
+				{
+					UE_LOG(LogMessageBusTester, Error, TEXT("With command line flag -MESSAGE_BUS_TESTER_PLANS, invalid interval specified -> %s"), *Tokens[1]);
+				}
+			}
+			else
+			{
+				UE_LOG(LogMessageBusTester, Error, TEXT("With command line flag -MESSAGE_BUS_TESTER_PLANS, invalid plan size specified -> %s"), *Tokens[0]);
+			}
+		}
+		else
+		{
+			UE_LOG(LogMessageBusTester, Error, TEXT("With command line flag -MESSAGE_BUS_TESTER_PLANS, invalid plan specified. Expected <size>:<interval> -> %s"), *Plan);
+		}
+	}
+
+	FString RunDuration;
+	if (FParse::Value(CommandLine, TEXT("-MESSAGE_BUS_TESTER_RUN_DURATION="), RunDuration))
+	{
+		const double InDuration = FCString::Atod(*RunDuration);
+		if (InDuration > 0)
+		{
+			TestDuration = InDuration;
+			UE_LOG(LogMessageBusTester, Display, TEXT("Setting run duration to '%f'"), TestDuration);
+		}
+	}
+
+	FString ExportToCSV;
+	if (FParse::Value(CommandLine, TEXT("-MESSAGE_BUS_TESTER_EXPORT_TO_CSV="), ExportToCSV))
+	{
+		ResultsToCSV = ExportToCSV;
+	}
+
 	// Only start the tester on initialization if the user has asked for it.
-	if (FParse::Param(FCommandLine::Get(), TEXT("MessageBusTesterStart")))
+	if (FParse::Param(FCommandLine::Get(), TEXT("-MESSAGE_BUS_TESTER_START_ON_LAUNCH")))
 	{
+		bStartOnLaunch = true;
 		Start();
 	}
 	
@@ -245,7 +323,11 @@ void FMessageBusTester::Start()
 
 		auto KeepAlive = [this](){SendKeepAlive_AnyThread();};
 		MessageBusHeartBeat = MakeUnique<FMessageBusTesterHeartBeat>(MoveTemp(KeepAlive));
-		AddTestPlanItem({ 10 * 1024 * 1024, 0.0f});
+
+		if (TestPlan.TestPlanItems.Num() == 0)
+		{
+			AddTestPlanItem({ 10 * 1024 * 1024, 0.0f});
+		}
 	}
 }
 
@@ -262,6 +344,110 @@ void FMessageBusTester::SendKeepAlive_AnyThread()
 	SendMessage(NewKeepAlive, EMessageBusTesterMessageFlags::None);
 }
 
+namespace UE::MessageBusTester::Private
+{
+FString GetCSVHeader()
+{
+	TStringBuilder<512> Builder;
+	Builder.Append(TEXT("UTC,"));
+	Builder.Append(TEXT("IPv4Address,"));
+	Builder.Append(TEXT("Bytes Sent,"));
+	Builder.Append(TEXT("Bytes Received,"));
+	Builder.Append(TEXT("Bytes Lost,"));
+	Builder.Append(TEXT("Percent Lost,"));
+	Builder.Append(TEXT("Average MB/s"));
+	return Builder.ToString();
+}
+
+void WriteTesterDataToCSV(const TSharedPtr<FDiscoveredTester, ESPMode::ThreadSafe>& Tester, TUniquePtr<FArchive>& CSVArchive)
+{
+	check(CSVArchive);
+
+	const FMessageTransportStatistics& Stat = Tester->Statistics;
+
+	TStringBuilder<512> Builder;
+	Builder.Append(FDateTime::UtcNow().ToString());
+	Builder.Append(TEXT(","));
+
+	Builder.Append(Stat.IPv4AsString);
+	Builder.Append(TEXT(","));
+
+	Builder.Appendf(TEXT("%d"), Stat.TotalBytesSent);
+	Builder.Append(TEXT(","));
+
+	Builder.Appendf(TEXT("%d"), Stat.TotalBytesReceived);
+	Builder.Append(TEXT(","));
+	
+	Builder.Appendf(TEXT("%d"), Stat.TotalBytesLost);
+	Builder.Append(TEXT(","));
+
+	double PercentLoss = 0;
+	if (Stat.TotalBytesSent > 0)
+	{
+		PercentLoss = 100 * (static_cast<double>(Stat.TotalBytesLost) / static_cast<double>(Stat.TotalBytesSent)); 
+	}
+	Builder.Appendf(TEXT("%.3f"), PercentLoss);
+	Builder.Append(TEXT(","));
+
+	Builder.Appendf(TEXT("%.3f"), Tester->AverageMbPerSecond.GetAverage());
+
+	FString CSVLine = Builder.ToString();
+	CSVLine += LINE_TERMINATOR;
+
+	FTCHARToUTF8 CSVLineUTF8(*CSVLine);
+	CSVArchive->Serialize((UTF8CHAR*)CSVLineUTF8.Get(), CSVLineUTF8.Length() * sizeof(UTF8CHAR));
+}
+	
+}
+void FMessageBusTester::ExportResultsToCSV()
+{
+	if (ResultsToCSV.IsEmpty())
+	{
+		return;
+	}
+
+	FString FullPath = FPaths::ConvertRelativePathToFull(ResultsToCSV);
+	if (FPaths::DirectoryExists(FullPath))
+	{
+		UE_LOG(LogMessageBusTester, Warning, TEXT("Invalid CSV filename ('%s') provided. Skipping export of CSV data."), *FullPath);
+		return;
+	}
+
+	FString FullPathDir = FPaths::GetPath(FullPath);
+	if (!IFileManager::Get().MakeDirectory(*FullPathDir, true))
+	{
+		UE_LOG(LogMessageBusTester, Error, TEXT("Error creating folders to host CSV ('%s')."), *FullPath);
+		return;
+	}
+
+	const bool bWriteHeader = !IFileManager::Get().FileExists(*FullPath);
+	TUniquePtr<FArchive> CSVArchive = TUniquePtr<FArchive>(IFileManager::Get().CreateFileWriter(*FullPath, EFileWrite::FILEWRITE_Append | EFileWrite::FILEWRITE_AllowRead));
+
+	if (!CSVArchive)
+	{
+		return;
+	}
+
+	if (bWriteHeader)
+	{
+		UTF8CHAR UTF8BOM[] = { (UTF8CHAR)0xEF, (UTF8CHAR)0xBB, (UTF8CHAR)0xBF };
+		CSVArchive->Serialize(&UTF8BOM, UE_ARRAY_COUNT(UTF8BOM) * sizeof(UTF8CHAR));
+		
+		FString CSVHeader = UE::MessageBusTester::Private::GetCSVHeader();
+		CSVHeader += LINE_TERMINATOR;
+		
+		FTCHARToUTF8 CSVHeaderUTF8(*CSVHeader);
+		CSVArchive->Serialize((UTF8CHAR*)CSVHeaderUTF8.Get(), CSVHeaderUTF8.Length() * sizeof(UTF8CHAR));
+	}
+
+	for (const TSharedPtr<FDiscoveredTester, ESPMode::ThreadSafe>& Tester : DiscoveredTesters)
+	{
+		UE::MessageBusTester::Private::WriteTesterDataToCSV(Tester, CSVArchive);
+	}
+
+	CSVArchive->Flush();	
+}
+
 void FMessageBusTester::Stop()
 {
 	StopTest();
@@ -278,6 +464,12 @@ void FMessageBusTester::Stop()
 		// Disable the Endpoint message handling since the message could keep it alive a bit.
 		TesterEndpoint->Disable();
 		TesterEndpoint.Reset();
+
+
+		// Now write the file results to a file.  It is possible for Stop to get called more than once soo
+		// this ensures that we only write the result once the test has been completed.
+		// 
+		ExportResultsToCSV();
 	}
 
 	// Clear the inbound message queue.
@@ -306,9 +498,34 @@ bool FMessageBusTester::Tick(float DeltaTime)
 		Visit(ProcessOne, Msg->Type);
 	}
 
+	if (bStartOnLaunch && GetState() != EMessageBusTesterState::Active)
+	{
+		FScopeLock Lock(&DiscoveredTesterCriticalSection);
+		if (DiscoveredTesters.Num() > 0)
+		{
+			StartTest();
+			// Reset our flag so we don't re-enter.
+			bStartOnLaunch = false;
+		}
+	}
+
 	if (GetState() == EMessageBusTesterState::Active)
-	{		
+	{
 		UpdateTestSequence();
+
+		if (StartTime > 0 && TestDuration > 0)
+		{
+			if (FPlatformTime::Seconds() > (StartTime + TestDuration))
+			{
+				UE_LOG(LogMessageBusTester, Display, TEXT("Automatically stopping test because we have exceeded our test duration."));
+				StopTest(true);
+			}
+		}
+	}
+
+	if (ShutdownTime > 0 && FPlatformTime::Seconds() > ShutdownTime)
+	{
+		RequestEngineExit(TEXT("Message bus tester shutdown requested"));		
 	}
 
 	return true;
@@ -572,6 +789,13 @@ void FMessageBusTester::ProcessOnGameThread(const FTesterStartTestMessage& Messa
 
 void FMessageBusTester::ProcessOnGameThread(const FTesterStopTestMessage& Message, const TSharedRef<IMessageContext, ESPMode::ThreadSafe>& Context)
 {
+	if (Message.bEngineShouldExit)
+	{
+		UE_LOG(LogMessageBusTester, Display, TEXT("Remote connection requested shutdown."));
+
+		TriggerAppExit();
+	}
+
 	SetState(EMessageBusTesterState::Idle);
 }
 
@@ -579,6 +803,12 @@ void FMessageBusTester::ProcessOnGameThread(const FTestPayloadReceptionMessage&
 {
 	if (TestPlanItemTrackers.IsValidIndex(Message.TestPlanIndex))
 	{
+		if (TestDuration > 0 && StartTime < 0)
+		{
+			// Only start our timer once we start receiving data.
+			StartTime = FPlatformTime::Seconds();
+		}
+
 		const FGuid& SenderIdentifier = Message.Identifier;
 		if (TSharedPtr<FDiscoveredTester, ESPMode::ThreadSafe>* FoundTesterPtr = DiscoveredTesters.FindByPredicate([SenderIdentifier](const TSharedPtr<FDiscoveredTester, ESPMode::ThreadSafe>& Other) { return Other->Identifier == SenderIdentifier; }))
 		{
diff --git a/Engine/Plugins/Experimental/MessageBusTester/Source/MessageBusTester/Private/MessageBusTester.h b/Engine/Plugins/Experimental/MessageBusTester/Source/MessageBusTester/Private/MessageBusTester.h
--- a/Engine/Plugins/Experimental/MessageBusTester/Source/MessageBusTester/Private/MessageBusTester.h
+++ b/Engine/Plugins/Experimental/MessageBusTester/Source/MessageBusTester/Private/MessageBusTester.h
@@ -76,7 +76,7 @@ public:
 	virtual void RemoveTestPlanItem(int32 Index) override;
 
 	virtual bool StartTest() override;
-	virtual bool StopTest() override;
+	virtual bool StopTest(bool bShouldExitOnStop = false) override;
 	virtual EMessageBusTesterState GetState() const override;
 	//~End IMessageBusTester
 
@@ -153,6 +153,12 @@ private:
 	void SetState(EMessageBusTesterState InState);
 
 private:
+	/** Dump test results to a CSV file. */
+	void ExportResultsToCSV();
+		
+	/** Set the shutdown timer to trigger this app to exit. */
+	void TriggerAppExit();
+
 	/** Messages received from message bus. */
 	TMpscQueue<FInboundMessage> InboundMessages;
 
@@ -193,6 +199,21 @@ private:
 
 	FMessageBusTestPlan TestPlan;
 
+	/** Test start time. If bStartOnLaunch is specified then this value will get set on first payload reception. */
+	double StartTime = -1;
+
+	/** Time to execute shutdown. If a test duration has been specified then a shutdown time will be set at the end of the duration. */
+	double ShutdownTime = -1;
+		
+	/** Run time duration in seconds. A number less than or equal to 0 indicates that we run forever.  If this is set we will also force the app to terminate. */
+	double TestDuration = -1;
+
+	/** Optional path to a CSV file that is written on exit. */
+	FString ResultsToCSV;
+
+	/** Indicates if we should execute our test plans as soon as we detect other clients. */
+	bool bStartOnLaunch = false;
+	
 	struct FTestPlanItemTracker
 	{
 		FTestPlanItem ItemDescription;
diff --git a/Engine/Plugins/Experimental/MessageBusTester/Source/MessageBusTester/Private/MessageBusTesterLogger.cpp b/Engine/Plugins/Experimental/MessageBusTester/Source/MessageBusTester/Private/MessageBusTesterLogger.cpp
--- a/Engine/Plugins/Experimental/MessageBusTester/Source/MessageBusTester/Private/MessageBusTesterLogger.cpp
+++ b/Engine/Plugins/Experimental/MessageBusTester/Source/MessageBusTester/Private/MessageBusTesterLogger.cpp
@@ -2,9 +2,23 @@
 
 #include "MessageBusTesterLogger.h"
 
+#include "MessageBusTesterModule.h"
 
 void FMessageBusTesterLogger::Log(const FName& TesterName, FString&& InLogMessage, EMessageSeverity::Type Verbosity /* = EMessageSeverity::Info */)
 {
+	if (Verbosity == EMessageSeverity::Info)
+	{
+		UE_LOG(LogMessageBusTester, Display, TEXT("%s"), *InLogMessage);
+	}
+	else if (Verbosity == EMessageSeverity::Warning)
+	{
+		UE_LOG(LogMessageBusTester, Warning, TEXT("%s"), *InLogMessage);
+	}
+	else if (Verbosity == EMessageSeverity::Error)
+	{
+		UE_LOG(LogMessageBusTester, Error, TEXT("%s"), *InLogMessage);
+	}
+
 	TSharedPtr<FMessageBusTesterLogEntry> NewEntry = MakeShared<FMessageBusTesterLogEntry>(TesterName, MoveTemp(InLogMessage), Verbosity);
 	Entries.Add(NewEntry);
 	OnMessageBusTesterNewLogReceived().Broadcast(NewEntry.ToSharedRef());
diff --git a/Engine/Plugins/Experimental/MessageBusTester/Source/MessageBusTester/Public/DiscoveredTester.h b/Engine/Plugins/Experimental/MessageBusTester/Source/MessageBusTester/Public/DiscoveredTester.h
--- a/Engine/Plugins/Experimental/MessageBusTester/Source/MessageBusTester/Public/DiscoveredTester.h
+++ b/Engine/Plugins/Experimental/MessageBusTester/Source/MessageBusTester/Public/DiscoveredTester.h
@@ -91,6 +91,10 @@ public:
 	/** Last interval between keep alive reception */
 	FMovingWindowAverageFast<double, 128> AverageUnreliableKeepAliveIntervalWindow;
 	FMovingWindowAverageFast<double, 128> AverageReliableKeepAliveIntervalWindow;
+
+	/** Running average of transfer performance. */
+	FMovingWindowAverageFast<double, 128> AverageMbPerSecond;
+		
 	FKeepAliveStatistics KeepAliveStatistics;
 
 	TArray<FPayloadInfo> SentPayloads;
diff --git a/Engine/Plugins/Experimental/MessageBusTester/Source/MessageBusTester/Public/IMessageBusTester.h b/Engine/Plugins/Experimental/MessageBusTester/Source/MessageBusTester/Public/IMessageBusTester.h
--- a/Engine/Plugins/Experimental/MessageBusTester/Source/MessageBusTester/Public/IMessageBusTester.h
+++ b/Engine/Plugins/Experimental/MessageBusTester/Source/MessageBusTester/Public/IMessageBusTester.h
@@ -33,7 +33,7 @@ public:
 	virtual bool StopSystem() = 0;
 
 	virtual bool StartTest() = 0;
-	virtual bool StopTest() = 0;
+	virtual bool StopTest(bool bShouldExitOnStop = false) = 0;
 	virtual void AddTestPlanItem(FTestPlanItem NewItem) = 0;
 	virtual void RemoveTestPlanItem(int32 Index) = 0;
 	virtual EMessageBusTesterState GetState() const = 0;
diff --git a/Engine/Plugins/Experimental/MessageBusTester/Source/MessageBusTester/Public/MessageBusTesterCommon.h b/Engine/Plugins/Experimental/MessageBusTester/Source/MessageBusTester/Public/MessageBusTesterCommon.h
--- a/Engine/Plugins/Experimental/MessageBusTester/Source/MessageBusTester/Public/MessageBusTesterCommon.h
+++ b/Engine/Plugins/Experimental/MessageBusTester/Source/MessageBusTester/Public/MessageBusTesterCommon.h
@@ -192,7 +192,8 @@ struct FTesterStopTestMessage : public FTesterBaseMessage
 public:
 	FTesterStopTestMessage() = default;
 
-public:
+	UPROPERTY()
+	bool bEngineShouldExit = false;
 };
 
 
diff --git a/Engine/Plugins/Experimental/MessageBusTester/Source/MessageBusTesterEditor/Private/SMessageBusTestNetwork.cpp b/Engine/Plugins/Experimental/MessageBusTester/Source/MessageBusTesterEditor/Private/SMessageBusTestNetwork.cpp
--- a/Engine/Plugins/Experimental/MessageBusTester/Source/MessageBusTesterEditor/Private/SMessageBusTestNetwork.cpp
+++ b/Engine/Plugins/Experimental/MessageBusTester/Source/MessageBusTesterEditor/Private/SMessageBusTestNetwork.cpp
@@ -383,7 +383,17 @@ FText SMessageBusTestNetwork::GetBytesSent() const
 FText SMessageBusTestNetwork::GetSegmentsLost() const
 {
 	FMessageTransportStatistics Stats = GetStats();
-	return FText::Format(LOCTEXT("BytesLost", "Bytes Lost : {0}"), FText::AsNumber(Stats.TotalBytesLost));
+	double PercentLoss = 0;
+
+	if (Stats.TotalBytesSent > 0)
+	{
+		PercentLoss = 100 * (static_cast<double>(Stats.TotalBytesLost) / static_cast<double>(Stats.TotalBytesSent)); 
+	}
+
+	FNumberFormattingOptions Options;
+	Options.MaximumFractionalDigits = 2;
+	Options.MinimumFractionalDigits = 2;
+	return FText::Format(LOCTEXT("BytesLost", "Bytes Lost : {0} ({1}%) "), FText::AsNumber(Stats.TotalBytesLost), FText::AsNumber(PercentLoss, &Options));
 }
 
 #undef LOCTEXT_NAMESPACE
