Fix the statistics panel not updating with a remote client resets its UDP Messaging settings.

#rnx

diff --git a/Engine/Plugins/Experimental/MessageBusTester/Source/MessageBusTesterEditor/Private/Widgets/SDiscoveredTesterListView.cpp b/Engine/Plugins/Experimental/MessageBusTester/Source/MessageBusTesterEditor/Private/Widgets/SDiscoveredTesterListView.cpp
--- a/Engine/Plugins/Experimental/MessageBusTester/Source/MessageBusTesterEditor/Private/Widgets/SDiscoveredTesterListView.cpp
+++ b/Engine/Plugins/Experimental/MessageBusTester/Source/MessageBusTesterEditor/Private/Widgets/SDiscoveredTesterListView.cpp
@@ -34,9 +34,10 @@ namespace DiscoveredTesterListView
 }
 
 FDiscoveredTesterTableRowData::FDiscoveredTesterTableRowData(
-	const FGuid& InIdentifier,	const FTesterInstanceDescriptor& InDescriptor,	const FMessageAddress& Address)
+	const FGuid& InIdentifier, const FTesterInstanceDescriptor& InDescriptor, const FMessageAddress& InAddress)
 	: Identifier(InIdentifier)
 	, Descriptor(InDescriptor)
+	, Address(InAddress)
 {
 	if (INetworkMessagingExtension* NetworkMessagingExtension = FMessageBusTesterEditorModule::GetMessagingStatistics())
 	{
@@ -93,6 +94,9 @@ void FDiscoveredTesterTableRowData::UpdateCachedValues()
 
 	if (INetworkMessagingExtension* NetworkMessagingExtension = FMessageBusTesterEditorModule::GetMessagingStatistics())
 	{
+		// The endpoint identifer can change if the remote host "resets" UdpMessaging. We always get the latest endpoint identifier so that
+		// we can properly track statistics. 
+		EndpointIdentifier = NetworkMessagingExtension->GetNodeIdFromAddress(Address);
 		Statistics = NetworkMessagingExtension->GetLatestNetworkStatistics(EndpointIdentifier);
 	}
 }
diff --git a/Engine/Plugins/Experimental/MessageBusTester/Source/MessageBusTesterEditor/Private/Widgets/SDiscoveredTesterListView.h b/Engine/Plugins/Experimental/MessageBusTester/Source/MessageBusTesterEditor/Private/Widgets/SDiscoveredTesterListView.h
--- a/Engine/Plugins/Experimental/MessageBusTester/Source/MessageBusTesterEditor/Private/Widgets/SDiscoveredTesterListView.h
+++ b/Engine/Plugins/Experimental/MessageBusTester/Source/MessageBusTesterEditor/Private/Widgets/SDiscoveredTesterListView.h
@@ -22,7 +22,7 @@ class IMessageBusTester;
  */
 struct FDiscoveredTesterTableRowData : TSharedFromThis<FDiscoveredTesterTableRowData>
 {
-	FDiscoveredTesterTableRowData(const FGuid& InIdentifier, const FTesterInstanceDescriptor& InDescriptor, const FMessageAddress& Address);
+	FDiscoveredTesterTableRowData(const FGuid& InIdentifier, const FTesterInstanceDescriptor& InDescriptor, const FMessageAddress& InAddress);
 
 	/** Fetch latest information for the associated data provider */
 	void UpdateCachedValues();
@@ -37,11 +37,12 @@ struct FDiscoveredTesterTableRowData : TSharedFromThis<FDiscoveredTesterTableRow
 	}
 
 public:
-
 	FGuid Identifier;
+	FTesterInstanceDescriptor Descriptor;
+
 	FGuid EndpointIdentifier;
+	FMessageAddress Address;
 
-	FTesterInstanceDescriptor Descriptor;
 	FMessageTransportStatistics Statistics;
 	EDiscoveredTesterConnectionState CachedConnectionState;
 	EMessageBusTesterState CachedTestingState;
